PR #6659: First version of RiderSourceCodeAccess plugin (Contributed by [REDACTED_PERSON])

diff --git a/Engine/Plugins/Developer/RiderSourceCodeAccess/RiderSourceCodeAccess.uplugin b/Engine/Plugins/Developer/RiderSourceCodeAccess/RiderSourceCodeAccess.uplugin
new file mode 100644
--- /dev/null
+++ b/Engine/Plugins/Developer/RiderSourceCodeAccess/RiderSourceCodeAccess.uplugin
@@ -0,0 +1,25 @@
+{
+	"FileVersion": 3,
+	"Version": 2,
+	"VersionName": "1.0",
+	"FriendlyName": "Rider Integration",
+	"Description": "Allows access to source code in Rider.",
+	"Category": "Programming",
+	"CreatedBy": "JetBrains",
+	"CreatedByURL": "https://www.jetbrains.com/",
+	"DocsURL": "",
+	"MarketplaceURL": "",
+	"SupportURL": "",
+	"EnabledByDefault": true,
+	"CanContainContent": false,
+	"IsBetaVersion": false,
+	"Installed": false,
+	"Modules": [
+		{
+			"Name": "RiderSourceCodeAccess",
+			"Type": "UncookedOnly",
+			"LoadingPhase": "Default",
+			"WhitelistPlatforms": [ "Win64" ]
+		}
+	]
+}
\ No newline at end of file
diff --git a/Engine/Plugins/Developer/RiderSourceCodeAccess/Source/RiderSourceCodeAccess/Private/RiderPathLocatorWin/RiderPathLocator.cpp b/Engine/Plugins/Developer/RiderSourceCodeAccess/Source/RiderSourceCodeAccess/Private/RiderPathLocatorWin/RiderPathLocator.cpp
new file mode 100644
--- /dev/null
+++ b/Engine/Plugins/Developer/RiderSourceCodeAccess/Source/RiderSourceCodeAccess/Private/RiderPathLocatorWin/RiderPathLocator.cpp
@@ -0,0 +1,205 @@
+// Copyright Epic Games, Inc. All Rights Reserved.
+
+#include "RiderPathLocator.h"
+
+#include "GenericPlatform/GenericPlatformFile.h"
+#include "HAL/FileManager.h"
+#include "Internationalization/Regex.h"
+#include "Misc/FileHelper.h"
+#include "Misc/Paths.h"
+#include "Dom/JsonObject.h"
+#include "Serialization/JsonReader.h"
+#include "Serialization/JsonSerializer.h"
+
+#include "Windows/WindowsPlatformMisc.h"
+#include "Windows/AllowWindowsPlatformTypes.h"
+#include <winreg.h>
+#include "Windows/HideWindowsPlatformTypes.h"
+
+namespace FRiderPathLocator
+{
+	static TArray<FInstallInfo> CollectPathsFromToolbox(const FString& IDEName)
+	{
+		FString ToolboxBinPath;
+
+		if (!FWindowsPlatformMisc::QueryRegKey(HKEY_CURRENT_USER, TEXT("Software\\JetBrains\\Toolbox\\"), TEXT(""), ToolboxBinPath)) return {};
+
+		FPaths::NormalizeDirectoryName(ToolboxBinPath);
+		const FString PatternString(TEXT("(.*)/bin"));
+		const FRegexPattern Pattern(PatternString);
+		FRegexMatcher ToolboxPathMatcher(Pattern, ToolboxBinPath);
+		if (!ToolboxPathMatcher.FindNext()) return {};
+
+		const FString ToolboxPath = ToolboxPathMatcher.GetCaptureGroup(1);
+		FString InstallPath = ToolboxPath;
+		const FString SettingJsonPath = FPaths::Combine(ToolboxPath, FString(".settings.json"));
+		if (FPaths::FileExists(SettingJsonPath))
+		{
+			FString JsonStr;
+			FFileHelper::LoadFileToString(JsonStr, *SettingJsonPath);
+			TSharedRef<TJsonReader<TCHAR>> JsonReader = TJsonReaderFactory<TCHAR>::Create(JsonStr);
+			TSharedPtr<FJsonObject> JsonObject = MakeShareable(new FJsonObject());
+			if (FJsonSerializer::Deserialize(JsonReader, JsonObject) && JsonObject.IsValid())
+			{
+				FString InstallLocation;
+				if (JsonObject->TryGetStringField(TEXT("install_location"), InstallLocation))
+				{
+					if (!InstallLocation.IsEmpty())
+					{
+						InstallPath = InstallLocation;
+					}
+				}
+			}
+		}
+		if(!FPaths::DirectoryExists(InstallPath)) return {};
+		
+		const FString ToolboxRiderRootPath = FPaths::Combine(InstallPath, FString("apps"), IDEName);
+		if(!FPaths::DirectoryExists(ToolboxRiderRootPath)) return {};
+
+		TArray<FInstallInfo> RiderInstallInfos;
+		TArray<FString> RiderPaths;
+		IFileManager::Get().FindFilesRecursive(RiderPaths, *ToolboxRiderRootPath, TEXT("rider64.exe"), true, false);
+		for(const FString& RiderPath: RiderPaths)
+		{
+			FRegexMatcher RiderPathMatcher(Pattern, RiderPath);
+			if (!RiderPathMatcher.FindNext()) continue;
+
+			const FString RiderDir = RiderPathMatcher.GetCaptureGroup(1);
+			const FString RiderCppPluginPath = FPaths::Combine(RiderDir, TEXT("plugins\\rider-cpp"));
+			if (FPaths::FileExists(RiderPath) && FPaths::DirectoryExists(RiderCppPluginPath))
+			{
+				FInstallInfo Info;
+				Info.Path = RiderPath;
+				Info.IsToolbox = true;
+				const FString BuildTxtPath = FPaths::Combine(RiderDir, TEXT("build.txt"));
+				if (FPaths::FileExists(BuildTxtPath))
+				{
+					FFileHelper::LoadFileToString(Info.Version, *BuildTxtPath);
+				}
+				else
+				{
+					Info.Version = FPaths::GetBaseFilename(RiderDir);
+				}
+				RiderInstallInfos.Add(Info);
+			}
+		}
+
+		return RiderInstallInfos;
+	}
+
+	static bool EnumerateRegistryKeys(HKEY Key, TArray<FString> &OutNames)
+	{
+		for (DWORD Index = 0;; Index++)
+		{
+			TCHAR KeyName[256];
+			DWORD KeyNameLength = sizeof(KeyName) / sizeof(KeyName[0]);
+
+			const LONG Result = RegEnumKeyEx(Key, Index, KeyName, &KeyNameLength, nullptr, nullptr, nullptr, nullptr);
+			if (Result == ERROR_NO_MORE_ITEMS)
+			{
+				break;
+			}
+
+			if (Result != ERROR_SUCCESS)
+			{
+				return false;
+			}
+
+			OutNames.Add(KeyName);
+		}
+		return true;
+	}
+
+	static bool EnumerateRegistryValues(HKEY Key, TArray<FString> &OutNames)
+	{
+		for (DWORD Index = 0;; Index++)
+		{
+			// Query the value
+			wchar_t ValueName[256];
+			DWORD ValueNameLength = sizeof(ValueName) / sizeof(ValueName[0]);
+
+			const LONG Result = RegEnumValue(Key, Index, ValueName, &ValueNameLength, nullptr, nullptr, nullptr, nullptr);
+			if (Result == ERROR_NO_MORE_ITEMS)
+			{
+				break;
+			}
+			
+			if (Result != ERROR_SUCCESS)
+			{
+				return false;
+			}
+
+			// Add it to the array
+			OutNames.Add(ValueName);
+		}
+		return true;
+	}
+
+	static LONG GetStringRegKey(const HKEY Key, const FString& ValueName, FString& Value)
+	{
+		WCHAR Buffer[512];
+		DWORD BufferSize = sizeof(Buffer);
+		const ULONG Result = RegQueryValueExW(Key, *ValueName, nullptr, nullptr, reinterpret_cast<LPBYTE>(Buffer), &BufferSize);
+		if (Result == ERROR_SUCCESS)
+		{
+			Value = Buffer;
+		}
+		return Result;
+	}
+
+	static TArray<FInstallInfo> CollectPathsFromRegistry(const FString& RegistryKey)
+	{
+		TArray<FInstallInfo> InstallInfos;
+		HKEY Key;
+		const LONG Result = RegOpenKeyEx(HKEY_LOCAL_MACHINE, *RegistryKey, 0, KEY_READ, &Key);
+		if (Result == ERROR_SUCCESS)
+		{
+			TArray<FString> Keys;
+			if (EnumerateRegistryKeys(Key, Keys))
+			{
+				for (const FString& key : Keys)
+				{
+					if (!key.Contains(TEXT("Rider"))) continue;
+
+					HKEY SubKey;
+					const LONG SubResult = RegOpenKeyEx(Key, *key, 0, KEY_READ, &SubKey);
+					if (SubResult != ERROR_SUCCESS) continue;
+
+					TArray<FString> Values;
+					if (!EnumerateRegistryValues(SubKey, Values)) continue;
+
+					for (const auto& Value : Values)
+					{
+						if (Value != TEXT("InstallLocation")) continue;
+						FString InstallLocation;
+						if (GetStringRegKey(SubKey, Value, InstallLocation) != ERROR_SUCCESS) continue;
+						const FString ExePath = FPaths::Combine(InstallLocation, TEXT("bin\\rider64.exe"));
+						const FString RiderCppPluginPath = FPaths::Combine(InstallLocation, TEXT("plugins\\rider-cpp"));
+						if (FPaths::FileExists(ExePath) && FPaths::DirectoryExists(RiderCppPluginPath))
+						{
+							FInstallInfo Info;
+							Info.IsToolbox = false;
+							Info.Path = ExePath;
+							GetStringRegKey(SubKey, TEXT("DisplayVersion"), Info.Version);
+							InstallInfos.Add(Info);
+						}
+					}
+
+				}
+			}
+		}
+
+		return InstallInfos;
+	}
+
+	TArray<FInstallInfo> CollectAllPaths()
+	{
+		TArray<FInstallInfo> InstallInfos;
+		InstallInfos.Append(CollectPathsFromToolbox(TEXT("Rider")));
+		InstallInfos.Append(CollectPathsFromToolbox(TEXT("RiderCpp")));
+		InstallInfos.Append(CollectPathsFromRegistry(TEXT("SOFTWARE\\Microsoft\\Windows\\CurrentVersion\\Uninstall")));
+		InstallInfos.Append(CollectPathsFromRegistry(TEXT("SOFTWARE\\WOW6432Node\\Microsoft\\Windows\\CurrentVersion\\Uninstall")));
+		return InstallInfos;
+	}
+}
+
diff --git a/Engine/Plugins/Developer/RiderSourceCodeAccess/Source/RiderSourceCodeAccess/Private/RiderPathLocatorWin/RiderPathLocator.h b/Engine/Plugins/Developer/RiderSourceCodeAccess/Source/RiderSourceCodeAccess/Private/RiderPathLocatorWin/RiderPathLocator.h
new file mode 100644
--- /dev/null
+++ b/Engine/Plugins/Developer/RiderSourceCodeAccess/Source/RiderSourceCodeAccess/Private/RiderPathLocatorWin/RiderPathLocator.h
@@ -0,0 +1,21 @@
+// Copyright Epic Games, Inc. All Rights Reserved.
+
+#pragma once
+
+#include "CoreMinimal.h"
+
+namespace FRiderPathLocator
+{
+	struct FInstallInfo
+	{
+		FString Path;
+		FString Version;
+		bool IsToolbox = false;
+
+		bool operator<(const FInstallInfo& rhs) const
+		{
+			return Version < rhs.Version;
+		}
+	};
+	TArray<FInstallInfo> CollectAllPaths();
+}
diff --git a/Engine/Plugins/Developer/RiderSourceCodeAccess/Source/RiderSourceCodeAccess/Private/RiderSourceCodeAccessor.cpp b/Engine/Plugins/Developer/RiderSourceCodeAccess/Source/RiderSourceCodeAccess/Private/RiderSourceCodeAccessor.cpp
new file mode 100644
--- /dev/null
+++ b/Engine/Plugins/Developer/RiderSourceCodeAccess/Source/RiderSourceCodeAccess/Private/RiderSourceCodeAccessor.cpp
@@ -0,0 +1,251 @@
+// Copyright Epic Games, Inc. All Rights Reserved.
+
+#include "RiderSourceCodeAccessor.h"
+
+#include "RiderPathLocator.h"
+
+#include "Modules/ModuleManager.h"
+#include "Misc/App.h"
+#include "Misc/FileHelper.h"
+#include "Misc/Paths.h"
+#include "Misc/ScopeLock.h"
+#include "Misc/UProjectInfo.h"
+#include "HAL/PlatformTime.h"
+#include "DesktopPlatformModule.h"
+#include "Interfaces/IProjectManager.h"
+#include "ProjectDescriptor.h"
+
+#define LOCTEXT_NAMESPACE "RiderSourceCodeAccessor"
+
+DEFINE_LOG_CATEGORY_STATIC(LogRiderAccessor, Log, All);
+
+namespace
+{
+TOptional<FString> ResolvePathToFile(const FString& FullPath)
+{
+	FString Path = FullPath;
+	if (FPaths::IsRelative(Path))
+		Path = FPaths::ConvertRelativePathToFull(Path);
+
+	if (!FPaths::FileExists(Path))
+	{
+		static const TArray<FString> SubDirs = {
+			"/Engine/Source/",
+			"/Engine/Plugins/"
+		};
+
+		FString EngineRootDir = FPaths::RootDir();
+		FPaths::NormalizeFilename(Path);
+		int32 Index = INDEX_NONE;
+		for (const FString& SubDir : SubDirs)
+		{
+			Index = Path.Find(SubDir);
+			if (Index != INDEX_NONE) break;
+		}
+		if (Index == INDEX_NONE) return {};
+
+		Path = EngineRootDir.Append(Path.RightChop(Index));
+		if (!FPaths::FileExists(Path)) return {};
+	}
+	return {Path};
+} 
+}
+
+void FRiderSourceCodeAccessor::RefreshAvailability()
+{
+	// If we have an executable path, we certainly have it installed!
+	bHasRiderInstalled = !ExecutablePath.IsEmpty() && FPaths::FileExists(ExecutablePath);
+}
+
+bool FRiderSourceCodeAccessor::AddSourceFiles(const TArray<FString>& AbsoluteSourcePaths, const TArray<FString>& AvailableModules)
+{
+	// @todo.Rider Manually add to folders? Or just regenerate
+	return false;
+}
+
+bool FRiderSourceCodeAccessor::CanAccessSourceCode() const
+{
+	return bHasRiderInstalled;
+}
+
+bool FRiderSourceCodeAccessor::DoesSolutionExist() const
+{
+	return !GetSolutionPath().IsEmpty();
+}
+
+FText FRiderSourceCodeAccessor::GetDescriptionText() const
+{
+	return LOCTEXT("RiderDisplayDesc", "Open source code files in Rider");
+}
+
+FName FRiderSourceCodeAccessor::GetFName() const
+{
+	return RiderName;
+}
+
+FText FRiderSourceCodeAccessor::GetNameText() const
+{
+	return FText::FromName(RiderName);
+}
+
+bool FRiderSourceCodeAccessor::OpenFileAtLine(const FString& FullPath, int32 LineNumber, int32)
+{
+	// UE_DEBUG_BREAK();
+	if (!bHasRiderInstalled) return false;
+	FString SolutionPath = GetSolutionPath();
+	if (SolutionPath.IsEmpty()) return false;
+	if (FPaths::IsRelative(SolutionPath))
+		SolutionPath = FPaths::ConvertRelativePathToFull(SolutionPath);
+
+	const TOptional<FString> OptionalPath = ResolvePathToFile(FullPath);
+	if(!OptionalPath.IsSet()) return false;
+	
+	const FString Path = OptionalPath.GetValue();
+	const FString Params = FString::Printf(TEXT("\"%s\" --line %d \"%s\""), *SolutionPath, LineNumber, *Path);
+
+	FProcHandle Proc = FPlatformProcess::CreateProc(*ExecutablePath, *Params, true, true, false, nullptr, 0, nullptr, nullptr);
+	if (!Proc.IsValid())
+	{
+		UE_LOG(LogRiderAccessor, Warning, TEXT("Opening file (%s) at a line (%d) failed."), *Path, LineNumber);
+		FPlatformProcess::CloseProc(Proc);
+		return false;
+	}
+
+	return true;
+}
+
+bool FRiderSourceCodeAccessor::OpenSolution()
+{
+	// UE_DEBUG_BREAK();
+	if (!bHasRiderInstalled) return false;
+
+	const FString SolutionPath = GetSolutionPath();
+	if (SolutionPath.IsEmpty()) return false;
+	
+	const FString FullPath = IFileManager::Get().ConvertToAbsolutePathForExternalAppForRead(*SolutionPath);
+	const FString Params = FString::Printf(TEXT("\"%s\""), *FullPath);
+	FPlatformProcess::CreateProc(*ExecutablePath, *Params, true, true, false, nullptr, 0, nullptr, nullptr);
+	return true;
+
+}
+
+bool FRiderSourceCodeAccessor::OpenSolutionAtPath(const FString& InSolutionPath)
+{
+	if (!bHasRiderInstalled) return false;
+
+	FString CorrectSolutionPath = InSolutionPath;
+	if (!CorrectSolutionPath.EndsWith(".sln"))
+	{
+		CorrectSolutionPath += ".sln";
+	}
+	const FString Params = FString::Printf(TEXT("\"%s\""), *CorrectSolutionPath);
+	FProcHandle Proc = FPlatformProcess::CreateProc(*ExecutablePath, *Params, true, true, false, nullptr, 0, nullptr, nullptr);
+	if (!Proc.IsValid())
+	{
+		UE_LOG(LogRiderAccessor, Warning, TEXT("Opening the project file (%s) failed."), *CorrectSolutionPath);
+		FPlatformProcess::CloseProc(Proc);
+		return false;
+	}
+	return true;
+}
+
+bool FRiderSourceCodeAccessor::OpenSourceFiles(const TArray<FString>& AbsoluteSourcePaths)
+{
+	// UE_DEBUG_BREAK();
+	if (!bHasRiderInstalled) return false;
+	FString SolutionPath = GetSolutionPath();
+	if (SolutionPath.IsEmpty()) return false;
+	
+	if (FPaths::IsRelative(SolutionPath))
+		SolutionPath = FPaths::ConvertRelativePathToFull(SolutionPath);
+
+	FString FilePaths = "";
+	for (const FString & FullPath : AbsoluteSourcePaths) {
+		const TOptional<FString> OptionalPath = ResolvePathToFile(FullPath);
+		if(!OptionalPath.IsSet()) return false;
+		const FString Path = OptionalPath.GetValue();
+		FilePaths += FString::Printf(TEXT("\"%s\" "), *Path);
+	}
+
+	const FString Params = FString::Printf(TEXT("\"%s\" %s"), *SolutionPath, *FilePaths);
+
+	
+
+	FProcHandle Proc = FPlatformProcess::CreateProc(*ExecutablePath, *Params, true, true, false, nullptr, 0, nullptr, nullptr);
+	if (!Proc.IsValid())
+	{
+		UE_LOG(LogRiderAccessor, Warning, TEXT("Opening files (%s) failed."), *FilePaths);
+		FPlatformProcess::CloseProc(Proc);
+		return false;
+	}
+	return true;
+}
+
+bool FRiderSourceCodeAccessor::SaveAllOpenDocuments() const
+{
+	return false;
+}
+
+void FRiderSourceCodeAccessor::Startup(const FRiderPathLocator::FInstallInfo& Info, ACCESS_TYPE Type)
+{	
+	ExecutablePath = Info.Path;
+	const FString IsToolboxText = Info.IsToolbox ? TEXT("(toolbox)") : TEXT("(installed)");
+	if(Type == ACCESS_TYPE::DIRECT)
+	{
+		RiderName = *FString::Format(TEXT("Rider {0} {1}"), { Info.Version, IsToolboxText });
+	}
+	else
+	{
+		RiderName = TEXT("Rider");
+	}
+	
+
+	RefreshAvailability();
+}
+
+
+FString FRiderSourceCodeAccessor::GetSolutionPath() const
+{
+	FScopeLock Lock(&CachedSolutionPathCriticalSection);
+
+	if(IsInGameThread())
+	{
+		if (CachedSolutionPathOverride.Len() > 0)
+		{
+			CachedSolutionPath = CachedSolutionPathOverride + TEXT(".sln");
+		}
+		else
+		{
+			CachedSolutionPath = FPaths::ConvertRelativePathToFull(FPaths::ProjectDir());
+
+			if (!FUProjectDictionary(FPaths::RootDir()).IsForeignProject(CachedSolutionPath))
+			{
+				FString MasterProjectName;
+				if (!FFileHelper::LoadFileToString(MasterProjectName, *(FPaths::EngineIntermediateDir() / TEXT("ProjectFiles/MasterProjectName.txt"))))
+				{
+					MasterProjectName = "UE4";
+				}
+				CachedSolutionPath = FPaths::Combine(FPaths::RootDir(), MasterProjectName + TEXT(".sln"));
+			}
+			else
+			{
+				const FProjectDescriptor* CurrentProject = IProjectManager::Get().GetCurrentProject();
+
+				if (CurrentProject == nullptr || CurrentProject->Modules.Num() == 0)
+				{
+					CachedSolutionPath = TEXT("");
+				}
+				else
+				{
+					const FString BaseName = FApp::HasProjectName() ? FApp::GetProjectName() : FPaths::GetBaseFilename(CachedSolutionPath);
+					CachedSolutionPath = FPaths::Combine(CachedSolutionPath, BaseName + TEXT(".sln"));
+				}
+			}
+		}
+	}
+
+	// This must be an absolute path as VS always uses absolute paths
+	return CachedSolutionPath;
+}
+
+#undef LOCTEXT_NAMESPACE
diff --git a/Engine/Plugins/Developer/RiderSourceCodeAccess/Source/RiderSourceCodeAccess/Private/RiderSourceCodeAccessor.h b/Engine/Plugins/Developer/RiderSourceCodeAccess/Source/RiderSourceCodeAccess/Private/RiderSourceCodeAccessor.h
new file mode 100644
--- /dev/null
+++ b/Engine/Plugins/Developer/RiderSourceCodeAccess/Source/RiderSourceCodeAccess/Private/RiderSourceCodeAccessor.h
@@ -0,0 +1,56 @@
+// Copyright Epic Games, Inc. All Rights Reserved.
+
+#pragma once
+
+#include "ISourceCodeAccessor.h"
+
+namespace FRiderPathLocator
+{
+	struct FInstallInfo;
+}
+
+class FRiderSourceCodeAccessor : public ISourceCodeAccessor
+{
+public:
+	enum class ACCESS_TYPE
+	{
+		DIRECT,
+		AGGREGATE
+	};
+	void Startup(const FRiderPathLocator::FInstallInfo & Info, ACCESS_TYPE Type = ACCESS_TYPE::DIRECT);
+
+	/** ISourceCodeAccessor implementation */
+	virtual void RefreshAvailability() override;
+	virtual bool CanAccessSourceCode() const override;
+	virtual bool DoesSolutionExist() const override;
+	virtual FName GetFName() const override;
+	virtual FText GetNameText() const override;
+	virtual FText GetDescriptionText() const override;
+	virtual bool OpenSolution() override;
+	virtual bool OpenSolutionAtPath(const FString& InSolutionPath) override;
+	virtual bool OpenFileAtLine(const FString& FullPath, int32 LineNumber, int32 ColumnNumber = 0) override;
+	virtual bool OpenSourceFiles(const TArray<FString>& AbsoluteSourcePaths) override;
+	virtual bool AddSourceFiles(const TArray<FString>& AbsoluteSourcePaths, const TArray<FString>& AvailableModules) override;
+	virtual bool SaveAllOpenDocuments() const override;
+	virtual void Tick(const float) override {}
+private:
+
+	FString GetSolutionPath() const;
+
+	FName RiderName;
+
+	/** Is Rider installed on this system? */
+	bool bHasRiderInstalled = false;
+	
+	/** The path to the Rider executable. */
+	FString ExecutablePath;
+
+	/** Critical section for updating SolutionPath */
+	mutable FCriticalSection CachedSolutionPathCriticalSection;
+
+	/** String storing the solution path obtained from the module manager to avoid having to use it on a thread */
+	mutable FString CachedSolutionPath = {};
+
+	/** Override for the cached solution path */
+	mutable FString CachedSolutionPathOverride = {};
+};
\ No newline at end of file
diff --git a/Engine/Plugins/Developer/RiderSourceCodeAccess/Source/RiderSourceCodeAccess/Private/RiderSourceCodeAccessorModule.cpp b/Engine/Plugins/Developer/RiderSourceCodeAccess/Source/RiderSourceCodeAccess/Private/RiderSourceCodeAccessorModule.cpp
new file mode 100644
--- /dev/null
+++ b/Engine/Plugins/Developer/RiderSourceCodeAccess/Source/RiderSourceCodeAccess/Private/RiderSourceCodeAccessorModule.cpp
@@ -0,0 +1,48 @@
+// Copyright Epic Games, Inc. All Rights Reserved.
+
+#include "RiderSourceCodeAccessorModule.h"
+
+#include "RiderPathLocator.h"
+
+#include "Modules/ModuleManager.h"
+#include "Features/IModularFeatures.h"
+
+#define LOCTEXT_NAMESPACE "RiderSourceCodeAccessor"
+
+IMPLEMENT_MODULE(FRiderSourceCodeAccessModule, RiderSourceCodeAccess);
+
+void FRiderSourceCodeAccessModule::ShutdownModule()
+{
+	for (auto& RiderSourceCodeAccessor : RiderSourceCodeAccessors)
+	{
+		// Unbind provider from editor
+		IModularFeatures::Get().UnregisterModularFeature(TEXT("SourceCodeAccessor"), &(RiderSourceCodeAccessor.Value.Get()));
+	}
+}
+
+void FRiderSourceCodeAccessModule::StartupModule()
+{
+	TArray<FRiderPathLocator::FInstallInfo> InstallInfos = FRiderPathLocator::CollectAllPaths();
+	InstallInfos.Sort();
+	for (const FRiderPathLocator::FInstallInfo & InstallInfo : InstallInfos)
+	{
+		TSharedRef<FRiderSourceCodeAccessor> RiderSourceCodeAccessor = MakeShareable(new FRiderSourceCodeAccessor());
+		RiderSourceCodeAccessor->Startup(InstallInfo);
+		IModularFeatures::Get().RegisterModularFeature(TEXT("SourceCodeAccessor"), &RiderSourceCodeAccessor.Get());
+		RiderSourceCodeAccessors.Add(RiderSourceCodeAccessor->GetFName(), RiderSourceCodeAccessor);
+	}
+	if(InstallInfos.Num() != 0)
+	{
+		TSharedRef<FRiderSourceCodeAccessor> RiderSourceCodeAccessor = MakeShareable(new FRiderSourceCodeAccessor());
+		RiderSourceCodeAccessor->Startup(InstallInfos.Last(), FRiderSourceCodeAccessor::ACCESS_TYPE::AGGREGATE);
+		IModularFeatures::Get().RegisterModularFeature(TEXT("SourceCodeAccessor"), &RiderSourceCodeAccessor.Get());
+		RiderSourceCodeAccessors.Add("Rider", RiderSourceCodeAccessor);		
+	}
+}
+
+bool FRiderSourceCodeAccessModule::SupportsDynamicReloading()
+{
+	return true;
+}
+
+#undef LOCTEXT_NAMESPACE
diff --git a/Engine/Plugins/Developer/RiderSourceCodeAccess/Source/RiderSourceCodeAccess/Private/RiderSourceCodeAccessorModule.h b/Engine/Plugins/Developer/RiderSourceCodeAccess/Source/RiderSourceCodeAccess/Private/RiderSourceCodeAccessorModule.h
new file mode 100644
--- /dev/null
+++ b/Engine/Plugins/Developer/RiderSourceCodeAccess/Source/RiderSourceCodeAccess/Private/RiderSourceCodeAccessorModule.h
@@ -0,0 +1,17 @@
+// Copyright Epic Games, Inc. All Rights Reserved.
+
+#pragma once
+
+#include "ISourceCodeAccessModule.h"
+#include "RiderSourceCodeAccessor.h"
+
+class FRiderSourceCodeAccessModule : public IModuleInterface
+{
+public:
+	/** IModuleInterface implementation */
+	virtual void StartupModule() override;
+	virtual void ShutdownModule() override;
+	virtual bool SupportsDynamicReloading() override;
+private:
+	TMap<FName, TSharedRef<FRiderSourceCodeAccessor>> RiderSourceCodeAccessors;
+};
diff --git a/Engine/Plugins/Developer/RiderSourceCodeAccess/Source/RiderSourceCodeAccess/RiderSourceCodeAccess.Build.cs b/Engine/Plugins/Developer/RiderSourceCodeAccess/Source/RiderSourceCodeAccess/RiderSourceCodeAccess.Build.cs
new file mode 100644
--- /dev/null
+++ b/Engine/Plugins/Developer/RiderSourceCodeAccess/Source/RiderSourceCodeAccess/RiderSourceCodeAccess.Build.cs
@@ -0,0 +1,27 @@
+// Copyright Epic Games, Inc. All Rights Reserved.
+
+namespace UnrealBuildTool.Rules
+{
+	public class RiderSourceCodeAccess : ModuleRules
+	{
+        public RiderSourceCodeAccess(ReadOnlyTargetRules Target) : base(Target)
+        {
+	        PCHUsage = PCHUsageMode.UseExplicitOrSharedPCHs;
+		    PrivateDependencyModuleNames.AddRange(
+                new []
+                {
+                    "Core",
+                    "SourceCodeAccess",
+                    "DesktopPlatform",
+					"Projects",
+					"Json"
+                }
+            );
+
+            if (Target.Platform == UnrealTargetPlatform.Win64)
+            {
+                PrivateIncludePaths.Add("RiderSourceCodeAccess/Private/RiderPathLocatorWin");
+            }
+		}
+	}
+}
\ No newline at end of file
