PR #6858: Fix discovery of "Rider for Unreal" version of Rider (Contributed by [REDACTED_PERSON])

diff --git a/Engine/Plugins/Developer/RiderSourceCodeAccess/Source/RiderSourceCodeAccess/Private/RiderPathLocatorWin/RiderPathLocator.cpp b/Engine/Plugins/Developer/RiderSourceCodeAccess/Source/RiderSourceCodeAccess/Private/RiderPathLocatorWin/RiderPathLocator.cpp
--- a/Engine/Plugins/Developer/RiderSourceCodeAccess/Source/RiderSourceCodeAccess/Private/RiderPathLocatorWin/RiderPathLocator.cpp
+++ b/Engine/Plugins/Developer/RiderSourceCodeAccess/Source/RiderSourceCodeAccess/Private/RiderPathLocatorWin/RiderPathLocator.cpp
@@ -18,7 +18,7 @@
 
 namespace FRiderPathLocator
 {
-	static TArray<FInstallInfo> CollectPathsFromToolbox(const FString& IDEName)
+	static TArray<FInstallInfo> CollectPathsFromToolbox()
 	{
 		FString ToolboxBinPath;
 
@@ -53,7 +53,7 @@ namespace FRiderPathLocator
 		}
 		if(!FPaths::DirectoryExists(InstallPath)) return {};
 		
-		const FString ToolboxRiderRootPath = FPaths::Combine(InstallPath, FString("apps"), IDEName);
+		const FString ToolboxRiderRootPath = FPaths::Combine(InstallPath, FString("apps"));
 		if(!FPaths::DirectoryExists(ToolboxRiderRootPath)) return {};
 
 		TArray<FInstallInfo> RiderInstallInfos;
@@ -71,12 +71,19 @@ namespace FRiderPathLocator
 				FInstallInfo Info;
 				Info.Path = RiderPath;
 				Info.IsToolbox = true;
-				const FString BuildTxtPath = FPaths::Combine(RiderDir, TEXT("build.txt"));
-				if (FPaths::FileExists(BuildTxtPath))
+				const FString ProductInfoJsonPath = FPaths::Combine(RiderDir, TEXT("product-info.json"));
+				if (FPaths::FileExists(ProductInfoJsonPath))
 				{
-					FFileHelper::LoadFileToString(Info.Version, *BuildTxtPath);
+					FString JsonStr;
+					FFileHelper::LoadFileToString(JsonStr, *ProductInfoJsonPath);
+					TSharedRef<TJsonReader<TCHAR>> JsonReader = TJsonReaderFactory<TCHAR>::Create(JsonStr);
+					TSharedPtr<FJsonObject> JsonObject = MakeShareable(new FJsonObject());
+					if (FJsonSerializer::Deserialize(JsonReader, JsonObject) && JsonObject.IsValid())
+					{
+						JsonObject->TryGetStringField(TEXT("buildNumber"), Info.Version);
+					}
 				}
-				else
+				if(Info.Version.IsEmpty())
 				{
 					Info.Version = FPaths::GetBaseFilename(RiderDir);
 				}
@@ -195,8 +202,7 @@ namespace FRiderPathLocator
 	TArray<FInstallInfo> CollectAllPaths()
 	{
 		TArray<FInstallInfo> InstallInfos;
-		InstallInfos.Append(CollectPathsFromToolbox(TEXT("Rider")));
-		InstallInfos.Append(CollectPathsFromToolbox(TEXT("RiderCpp")));
+		InstallInfos.Append(CollectPathsFromToolbox());
 		InstallInfos.Append(CollectPathsFromRegistry(TEXT("SOFTWARE\\Microsoft\\Windows\\CurrentVersion\\Uninstall")));
 		InstallInfos.Append(CollectPathsFromRegistry(TEXT("SOFTWARE\\WOW6432Node\\Microsoft\\Windows\\CurrentVersion\\Uninstall")));
 		return InstallInfos;
