- Disabled Rider installation discovery from doing recursive search through /Applications because it was taking _forever_ and also asking for permissions as we touch paths and maybe follow symlinks

diff --git a/Engine/Plugins/Developer/RiderSourceCodeAccess/Source/RiderSourceCodeAccess/Private/RiderPathLocator/Mac/RiderPathLocatorMac.cpp b/Engine/Plugins/Developer/RiderSourceCodeAccess/Source/RiderSourceCodeAccess/Private/RiderPathLocator/Mac/RiderPathLocatorMac.cpp
--- a/Engine/Plugins/Developer/RiderSourceCodeAccess/Source/RiderSourceCodeAccess/Private/RiderPathLocator/Mac/RiderPathLocatorMac.cpp
+++ b/Engine/Plugins/Developer/RiderSourceCodeAccess/Source/RiderSourceCodeAccess/Private/RiderPathLocator/Mac/RiderPathLocatorMac.cpp
@@ -31,7 +31,7 @@ static TArray<FInstallInfo> GetManuallyInstalledRiders()
 {
 	TArray<FInstallInfo> Result;
 	TArray<FString> RiderPaths;
-	IFileManager::Get().FindFilesRecursive(RiderPaths, TEXT("/Applications"), TEXT("Rider*.app"), false, true);
+	IFileManager::Get().FindFiles(RiderPaths, TEXT("/Applications/Rider*.app"), false, true);
 	for(const FString& RiderPath: RiderPaths)
 	{
 		TOptional<FInstallInfo> InstallInfo = FRiderPathLocator::GetInstallInfoFromRiderPath(RiderPath, false);
@@ -63,4 +63,4 @@ TSet<FInstallInfo> FRiderPathLocator::CollectAllPaths()
 	InstallInfos.Append(GetInstallInfosFromToolbox(GetToolboxPath(), "Rider*.app"));
 	return InstallInfos;
 }
-#endif
\ No newline at end of file
+#endif
