Avoid Rider spamming log if mdfind is not installed on the system
Fix some coding standard issues

diff --git a/Engine/Plugins/Developer/RiderSourceCodeAccess/Source/RiderSourceCodeAccess/Private/RiderPathLocator/Linux/RiderPathLocatorLinux.cpp b/Engine/Plugins/Developer/RiderSourceCodeAccess/Source/RiderSourceCodeAccess/Private/RiderPathLocator/Linux/RiderPathLocatorLinux.cpp
--- a/Engine/Plugins/Developer/RiderSourceCodeAccess/Source/RiderSourceCodeAccess/Private/RiderPathLocator/Linux/RiderPathLocatorLinux.cpp
+++ b/Engine/Plugins/Developer/RiderSourceCodeAccess/Source/RiderSourceCodeAccess/Private/RiderPathLocator/Linux/RiderPathLocatorLinux.cpp
@@ -13,12 +13,18 @@
 
 TOptional<FInstallInfo> FRiderPathLocator::GetInstallInfoFromRiderPath(const FString& PathToRiderApp, FInstallInfo::EInstallType InstallType)
 {
-	if(!FPaths::DirectoryExists(PathToRiderApp)) return {};
+	if(!FPaths::DirectoryExists(PathToRiderApp))
+	{
+		return {};
+	}
 
 	const FString RiderCppPluginPath = FPaths::Combine(PathToRiderApp, TEXT("plugins"), TEXT("rider-cpp"));
 
-	if (!FPaths::DirectoryExists(RiderCppPluginPath)) return {};
-	
+	if (!FPaths::DirectoryExists(RiderCppPluginPath))
+	{
+		return {};
+	}
+
 	FInstallInfo Info;
 	Info.Path = FPaths::Combine(PathToRiderApp, TEXT("bin"), TEXT("rider.sh"));
 	Info.InstallType = InstallType;
@@ -50,7 +56,7 @@ static TArray<FInstallInfo> GetManuallyInstalledRiders()
 	const FString FHomePath = GetHomePath();
 
 	const FString LocalPathMask = FPaths::Combine(FHomePath, TEXT("Rider*"));
-	
+
 	IFileManager::Get().FindFiles(RiderPaths, *LocalPathMask, false, true);
 
 	for(const FString& RiderPath: RiderPaths)
@@ -58,7 +64,9 @@ static TArray<FInstallInfo> GetManuallyInstalledRiders()
 		FString FullPath = FPaths::Combine(FHomePath, RiderPath);
 		TOptional<FInstallInfo> InstallInfo = FRiderPathLocator::GetInstallInfoFromRiderPath(FullPath, FInstallInfo::EInstallType::Installed);
 		if(InstallInfo.IsSet())
+		{
 			Result.Add(InstallInfo.GetValue());
+		}
 	}
 	return Result;
 }
@@ -73,28 +81,41 @@ static FString GetToolboxPath()
 
 static TArray<FInstallInfo> GetInstalledRidersWithMdfind()
 {
-    int32 ReturnCode;
-    FString OutResults;
-    FString OutErrors;
-    FPlatformProcess::ExecProcess(TEXT("/usr/bin/mdfind"), TEXT("\"kMDItemKind == Application\""), &ReturnCode, &OutResults, &OutErrors);
-    if (ReturnCode != 0)
+	int32 ReturnCode;
+	FString OutResults;
+	FString OutErrors;
+
+	// avoid trying to run mdfind if it doesnt exists
+	if (!FPaths::FileExists(TEXT("/usr/bin/mdfind")))
+	{
 		return {};
+	}
 
-    TArray<FString> RiderPaths;
+	FPlatformProcess::ExecProcess(TEXT("/usr/bin/mdfind"), TEXT("\"kMDItemKind == Application\""), &ReturnCode, &OutResults, &OutErrors);
+	if (ReturnCode != 0)
+	{
+		return {};
+	}
+
+	TArray<FString> RiderPaths;
 	FString TmpString;
 	while(OutResults.Split(TEXT("\n"), &TmpString, &OutResults))
 	{
 		if(TmpString.Contains(TEXT("Rider")))
+		{
 			RiderPaths.Add(TmpString);
+		}
 	}
-    TArray<FInstallInfo> Result;
-    for(const FString& RiderPath: RiderPaths)
-    {
-        TOptional<FInstallInfo> InstallInfo = FRiderPathLocator::GetInstallInfoFromRiderPath(RiderPath, FInstallInfo::EInstallType::Installed);
-        if(InstallInfo.IsSet())
-            Result.Add(InstallInfo.GetValue());
-    }
-    return Result;
+	TArray<FInstallInfo> Result;
+	for(const FString& RiderPath: RiderPaths)
+	{
+		TOptional<FInstallInfo> InstallInfo = FRiderPathLocator::GetInstallInfoFromRiderPath(RiderPath, FInstallInfo::EInstallType::Installed);
+		if(InstallInfo.IsSet())
+		{
+			Result.Add(InstallInfo.GetValue());
+		}
+	}
+	return Result;
 }
 
 TSet<FInstallInfo> FRiderPathLocator::CollectAllPaths()
diff --git a/Engine/Plugins/Developer/RiderSourceCodeAccess/Source/RiderSourceCodeAccess/Private/RiderPathLocator/Mac/RiderPathLocatorMac.cpp b/Engine/Plugins/Developer/RiderSourceCodeAccess/Source/RiderSourceCodeAccess/Private/RiderPathLocator/Mac/RiderPathLocatorMac.cpp
--- a/Engine/Plugins/Developer/RiderSourceCodeAccess/Source/RiderSourceCodeAccess/Private/RiderPathLocator/Mac/RiderPathLocatorMac.cpp
+++ b/Engine/Plugins/Developer/RiderSourceCodeAccess/Source/RiderSourceCodeAccess/Private/RiderPathLocator/Mac/RiderPathLocatorMac.cpp
@@ -13,11 +13,17 @@
 
 TOptional<FInstallInfo> FRiderPathLocator::GetInstallInfoFromRiderPath(const FString& PathToRiderApp, FInstallInfo::EInstallType InstallType)
 {
-	if(!FPaths::DirectoryExists(PathToRiderApp)) return {};
+	if(!FPaths::DirectoryExists(PathToRiderApp))
+	{
+		return {};
+	}
 
 	const FString RiderCppPluginPath = FPaths::Combine(PathToRiderApp, TEXT("Contents"), TEXT("plugins"), TEXT("rider-cpp"));
-	if (!FPaths::DirectoryExists(RiderCppPluginPath)) return {};
-	
+	if (!FPaths::DirectoryExists(RiderCppPluginPath))
+	{
+		return {};
+	}
+
 	FInstallInfo Info;
 	Info.Path = FPaths::Combine(PathToRiderApp, TEXT("Contents"), TEXT("MacOS"), TEXT("rider"));
 	Info.InstallType = InstallType;
@@ -39,7 +45,9 @@ static TArray<FInstallInfo> GetManuallyInstalledRiders()
 		FString FullPath = TEXT("/Applications/") + RiderPath;
 		TOptional<FInstallInfo> InstallInfo = FRiderPathLocator::GetInstallInfoFromRiderPath(FullPath, FInstallInfo::EInstallType::Installed);
 		if(InstallInfo.IsSet())
+		{
 			Result.Add(InstallInfo.GetValue());
+		}
 	}
 	return Result;
 }
@@ -61,28 +69,41 @@ static FString GetToolboxPath()
 
 static TArray<FInstallInfo> GetInstalledRidersWithMdfind()
 {
-    int32 ReturnCode;
-    FString OutResults;
-    FString OutErrors;
-    FPlatformProcess::ExecProcess(TEXT("/usr/bin/mdfind"), TEXT("\"kMDItemKind == Application\""), &ReturnCode, &OutResults, &OutErrors);
-    if (ReturnCode != 0)
+	int32 ReturnCode;
+	FString OutResults;
+	FString OutErrors;
+
+	// avoid trying to run mdfind if it doesnt exists
+	if (!FPaths::FileExists(TEXT("/usr/bin/mdfind")))
+	{
+		return {};
+	}
+
+	FPlatformProcess::ExecProcess(TEXT("/usr/bin/mdfind"), TEXT("\"kMDItemKind == Application\""), &ReturnCode, &OutResults, &OutErrors);
+	if (ReturnCode != 0)
+	{
 		return {};
+	}
 
-    TArray<FString> RiderPaths;
+	TArray<FString> RiderPaths;
 	FString TmpString;
 	while(OutResults.Split(TEXT("\n"), &TmpString, &OutResults))
 	{
 		if(TmpString.Contains(TEXT("Rider")))
+		{
 			RiderPaths.Add(TmpString);
+		}
 	}
-    TArray<FInstallInfo> Result;
-    for(const FString& RiderPath: RiderPaths)
-    {
-        TOptional<FInstallInfo> InstallInfo = FRiderPathLocator::GetInstallInfoFromRiderPath(RiderPath, FInstallInfo::EInstallType::Installed);
-        if(InstallInfo.IsSet())
-            Result.Add(InstallInfo.GetValue());
-    }
-    return Result;
+	TArray<FInstallInfo> Result;
+	for(const FString& RiderPath: RiderPaths)
+	{
+		TOptional<FInstallInfo> InstallInfo = FRiderPathLocator::GetInstallInfoFromRiderPath(RiderPath, FInstallInfo::EInstallType::Installed);
+		if(InstallInfo.IsSet())
+		{
+			Result.Add(InstallInfo.GetValue());
+		}
+	}
+	return Result;
 }
 
 TSet<FInstallInfo> FRiderPathLocator::CollectAllPaths()
