PR #8698: Merge RiderSourceCodeAccess changes from GitHub (Contributed by [REDACTED_PERSON])

diff --git a/Engine/Plugins/Developer/RiderSourceCodeAccess/RiderSourceCodeAccess.uplugin b/Engine/Plugins/Developer/RiderSourceCodeAccess/RiderSourceCodeAccess.uplugin
--- a/Engine/Plugins/Developer/RiderSourceCodeAccess/RiderSourceCodeAccess.uplugin
+++ b/Engine/Plugins/Developer/RiderSourceCodeAccess/RiderSourceCodeAccess.uplugin
@@ -1,7 +1,7 @@
 {
 	"FileVersion": 3,
 	"Version": 2,
-	"VersionName": "1.4",
+	"VersionName": "1.5",
 	"FriendlyName": "Rider Integration",
 	"Description": "Allows access to source code in Rider.",
 	"Category": "Programming",
diff --git a/Engine/Plugins/Developer/RiderSourceCodeAccess/Source/RiderSourceCodeAccess/Private/RiderPathLocator/Linux/RiderPathLocatorLinux.cpp b/Engine/Plugins/Developer/RiderSourceCodeAccess/Source/RiderSourceCodeAccess/Private/RiderPathLocator/Linux/RiderPathLocatorLinux.cpp
--- a/Engine/Plugins/Developer/RiderSourceCodeAccess/Source/RiderSourceCodeAccess/Private/RiderPathLocator/Linux/RiderPathLocatorLinux.cpp
+++ b/Engine/Plugins/Developer/RiderSourceCodeAccess/Source/RiderSourceCodeAccess/Private/RiderPathLocator/Linux/RiderPathLocatorLinux.cpp
@@ -11,28 +11,40 @@
 
 #if PLATFORM_LINUX
 
-TOptional<FInstallInfo> FRiderPathLocator::GetInstallInfoFromRiderPath(const FString& PathToRiderApp, FInstallInfo::EInstallType InstallType)
+TOptional<FInstallInfo> FRiderPathLocator::GetInstallInfoFromRiderPath(const FString& Path, FInstallInfo::EInstallType InstallType)
 {
-	if(!FPaths::DirectoryExists(PathToRiderApp))
+	if(!FPaths::FileExists(Path))
+	{
+		return {};
+	}
+	
+	const FString PatternString(TEXT("(.*)(?:\\\\|/)bin"));
+	const FRegexPattern Pattern(PatternString);
+	FRegexMatcher RiderPathMatcher(Pattern, Path);
+	if (!RiderPathMatcher.FindNext())
 	{
 		return {};
 	}
 
-	const FString RiderCppPluginPath = FPaths::Combine(PathToRiderApp, TEXT("plugins"), TEXT("rider-cpp"));
-
+	const FString RiderDir = RiderPathMatcher.GetCaptureGroup(1);
+	const FString RiderCppPluginPath = FPaths::Combine(RiderDir, TEXT("plugins"), TEXT("rider-cpp"));
 	if (!FPaths::DirectoryExists(RiderCppPluginPath))
 	{
 		return {};
 	}
-
+	
 	FInstallInfo Info;
-	Info.Path = FPaths::Combine(PathToRiderApp, TEXT("bin"), TEXT("rider.sh"));
+	Info.Path = Path;
 	Info.InstallType = InstallType;
-	const FString ProductInfoJsonPath = FPaths::Combine(PathToRiderApp, TEXT("product-info.json"));
+	const FString ProductInfoJsonPath = FPaths::Combine(RiderDir, TEXT("product-info.json"));
 	if (FPaths::FileExists(ProductInfoJsonPath))
 	{
 		ParseProductInfoJson(Info, ProductInfoJsonPath);
 	}
+	if(!Info.Version.IsInitialized())
+	{
+		Info.Version = FPaths::GetBaseFilename(RiderDir);
+	}
 	return Info;
 }
 
diff --git a/Engine/Plugins/Developer/RiderSourceCodeAccess/Source/RiderSourceCodeAccess/Private/RiderPathLocator/Win/RiderPathLocatorWin.cpp b/Engine/Plugins/Developer/RiderSourceCodeAccess/Source/RiderSourceCodeAccess/Private/RiderPathLocator/Win/RiderPathLocatorWin.cpp
--- a/Engine/Plugins/Developer/RiderSourceCodeAccess/Source/RiderSourceCodeAccess/Private/RiderPathLocator/Win/RiderPathLocatorWin.cpp
+++ b/Engine/Plugins/Developer/RiderSourceCodeAccess/Source/RiderSourceCodeAccess/Private/RiderPathLocator/Win/RiderPathLocatorWin.cpp
@@ -14,11 +14,25 @@
 #include <winreg.h>
 #include "Windows/HideWindowsPlatformTypes.h"
 
-static FString GetToolboxPath(const Windows::HKEY RootKey)
+static FString GetToolboxPath()
+{	
+	TArray<FInstallInfo> Result;
+#if ENGINE_MAJOR_VERSION == 4 && ENGINE_MINOR_VERSION <= 20
+	TCHAR CHomePath[4096];
+	FPlatformMisc::GetEnvironmentVariable(TEXT("LOCALAPPDATA"), CHomePath, ARRAY_COUNT(CHomePath));
+	const FString LocalAppData = CHomePath;
+#else
+	const FString LocalAppData = FPlatformMisc::GetEnvironmentVariable(TEXT("LOCALAPPDATA"));
+#endif
+
+	return FPaths::Combine(LocalAppData, TEXT("JetBrains"), TEXT("Toolbox"));
+}
+
+static FString GetToolboxPath(const Windows::HKEY RootKey, const FString& RegistryKey)
 {
 	FString ToolboxBinPath;
 
-	if (!FWindowsPlatformMisc::QueryRegKey(RootKey, TEXT("Software\\JetBrains\\Toolbox\\"), TEXT(""), ToolboxBinPath)) return {};
+	if (!FWindowsPlatformMisc::QueryRegKey(RootKey, *RegistryKey, TEXT(""), ToolboxBinPath)) return {};
 
 	FPaths::NormalizeDirectoryName(ToolboxBinPath);
 	const FString PatternString(TEXT("(.*)(?:\\\\|/)bin"));
@@ -159,8 +173,11 @@ TOptional<FInstallInfo> FRiderPathLocator::GetInstallInfoFromRiderPath(const FSt
 TSet<FInstallInfo> FRiderPathLocator::CollectAllPaths()
 {
 	TSet<FInstallInfo> InstallInfos;
-	InstallInfos.Append(GetInstallInfosFromToolbox(GetToolboxPath(HKEY_CURRENT_USER), "rider64.exe"));
-	InstallInfos.Append(GetInstallInfosFromToolbox(GetToolboxPath(HKEY_LOCAL_MACHINE), "rider64.exe"));
+	InstallInfos.Append(GetInstallInfosFromToolbox(GetToolboxPath(), "rider64.exe"));
+	InstallInfos.Append(GetInstallInfosFromToolbox(GetToolboxPath(HKEY_CURRENT_USER, TEXT("Software\\JetBrains\\Toolbox\\")), "rider64.exe"));
+	InstallInfos.Append(GetInstallInfosFromToolbox(GetToolboxPath(HKEY_LOCAL_MACHINE, TEXT("Software\\JetBrains\\Toolbox\\")), "rider64.exe"));
+	InstallInfos.Append(GetInstallInfosFromToolbox(GetToolboxPath(HKEY_CURRENT_USER, TEXT("Software\\JetBrains s.r.o.\\JetBrainsToolbox\\")), "rider64.exe"));
+	InstallInfos.Append(GetInstallInfosFromToolbox(GetToolboxPath(HKEY_LOCAL_MACHINE, TEXT("Software\\JetBrains s.r.o.\\JetBrainsToolbox\\")), "rider64.exe"));
 	InstallInfos.Append(CollectPathsFromRegistry(HKEY_CURRENT_USER, TEXT("SOFTWARE\\Microsoft\\Windows\\CurrentVersion\\Uninstall")));
 	InstallInfos.Append(CollectPathsFromRegistry(HKEY_LOCAL_MACHINE, TEXT("SOFTWARE\\Microsoft\\Windows\\CurrentVersion\\Uninstall")));
 	InstallInfos.Append(CollectPathsFromRegistry(HKEY_CURRENT_USER, TEXT("SOFTWARE\\WOW6432Node\\Microsoft\\Windows\\CurrentVersion\\Uninstall")));
