PR #8780: Merging changes from the GitHub RSCA (Contributed by [REDACTED_PERSON])

diff --git a/Engine/Plugins/Developer/RiderSourceCodeAccess/Source/RiderSourceCodeAccess/Private/RiderPathLocator/Common/RiderPathLocator.cpp b/Engine/Plugins/Developer/RiderSourceCodeAccess/Source/RiderSourceCodeAccess/Private/RiderPathLocator/Common/RiderPathLocator.cpp
--- a/Engine/Plugins/Developer/RiderSourceCodeAccess/Source/RiderSourceCodeAccess/Private/RiderPathLocator/Common/RiderPathLocator.cpp
+++ b/Engine/Plugins/Developer/RiderSourceCodeAccess/Source/RiderSourceCodeAccess/Private/RiderPathLocator/Common/RiderPathLocator.cpp
@@ -47,7 +47,9 @@ TArray<FInstallInfo> FRiderPathLocator::GetInstallInfos(const FString& ToolboxRi
 	{
 		TOptional<FInstallInfo> InstallInfo = GetInstallInfoFromRiderPath(RiderPath, InstallType);
 		if(InstallInfo.IsSet())
+		{
 			RiderInstallInfos.Add(InstallInfo.GetValue());
+		}
 	}
 	return RiderInstallInfos;
 }
@@ -63,6 +65,12 @@ void FRiderPathLocator::ParseProductInfoJson(FInstallInfo& Info, const FString&
 		FString VersionString;
 		JsonObject->TryGetStringField(TEXT("buildNumber"), VersionString);
 		Info.Version = VersionString;
+		if(Info.Version.Major() >= 221)
+		{
+			Info.SupportUprojectState = FInstallInfo::ESupportUproject::Release;
+			return;
+		}
+
 		const TArray< TSharedPtr<FJsonValue> >* CustomProperties;
 		if(!JsonObject->TryGetArrayField(TEXT("customProperties"), CustomProperties)) return;
 
@@ -74,7 +82,8 @@ void FRiderPathLocator::ParseProductInfoJson(FInstallInfo& Info, const FString&
 			FString SupportUprojectStateKey;
 			const bool bIsValidKey = Item->TryGetStringField(TEXT("key"), SupportUprojectStateKey);
 			if(!bIsValidKey) continue;
-			if(!SupportUprojectStateKey.Equals(TEXT("SupportUproject"))) continue;
+			if(	!SupportUprojectStateKey.Equals(TEXT("SupportUproject")) &&
+				!SupportUprojectStateKey.Equals(TEXT("SupportUprojectState"))) continue;
 
 			FString SupportUprojectStateValue;
 			const bool bIsValidValue  = Item->TryGetStringField(TEXT("value"), SupportUprojectStateValue);
@@ -103,7 +112,9 @@ TArray<FInstallInfo> FRiderPathLocator::GetInstallInfosFromResourceFile()
 		
 		TOptional<FInstallInfo> InstallInfo = GetInstallInfoFromRiderPath(Location, FInstallInfo::EInstallType::Custom);
 		if(InstallInfo.IsSet())
+		{
 			RiderInstallInfos.Add(InstallInfo.GetValue());
+		}
 	}
 	return RiderInstallInfos;
 }
diff --git a/Engine/Plugins/Developer/RiderSourceCodeAccess/Source/RiderSourceCodeAccess/Private/RiderPathLocator/Linux/RiderPathLocatorLinux.cpp b/Engine/Plugins/Developer/RiderSourceCodeAccess/Source/RiderSourceCodeAccess/Private/RiderPathLocator/Linux/RiderPathLocatorLinux.cpp
--- a/Engine/Plugins/Developer/RiderSourceCodeAccess/Source/RiderSourceCodeAccess/Private/RiderPathLocator/Linux/RiderPathLocatorLinux.cpp
+++ b/Engine/Plugins/Developer/RiderSourceCodeAccess/Source/RiderSourceCodeAccess/Private/RiderPathLocator/Linux/RiderPathLocatorLinux.cpp
@@ -44,6 +44,10 @@ TOptional<FInstallInfo> FRiderPathLocator::GetInstallInfoFromRiderPath(const FSt
 	if(!Info.Version.IsInitialized())
 	{
 		Info.Version = FPaths::GetBaseFilename(RiderDir);
+		if(Info.Version.Major() >= 221)
+		{
+			Info.SupportUprojectState = FInstallInfo::ESupportUproject::Release;
+		}
 	}
 	return Info;
 }
@@ -65,15 +69,30 @@ static TArray<FInstallInfo> GetManuallyInstalledRiders()
 {
 	TArray<FInstallInfo> Result;
 	TArray<FString> RiderPaths;
+
 	const FString FHomePath = GetHomePath();
+	const FString HomePathMask = FPaths::Combine(FHomePath, TEXT("Rider*"));
+	
+	IFileManager::Get().FindFiles(RiderPaths, *HomePathMask, false, true);
 
-	const FString LocalPathMask = FPaths::Combine(FHomePath, TEXT("Rider*"));
+	for(const FString& RiderPath: RiderPaths)
+	{
+		FString FullPath = FPaths::Combine(FHomePath, RiderPath);
+		TOptional<FInstallInfo> InstallInfo = FRiderPathLocator::GetInstallInfoFromRiderPath(FullPath, FInstallInfo::EInstallType::Installed);
+		if(InstallInfo.IsSet())
+		{
+			Result.Add(InstallInfo.GetValue());
+		}
+	}
 
-	IFileManager::Get().FindFiles(RiderPaths, *LocalPathMask, false, true);
+	const FString FOptPath = TEXT("/opt");
+	const FString OptPathMask = FPaths::Combine(FOptPath, TEXT("Rider*"));
+	
+	IFileManager::Get().FindFiles(RiderPaths, *OptPathMask, false, true);
 
 	for(const FString& RiderPath: RiderPaths)
 	{
-		FString FullPath = FPaths::Combine(FHomePath, RiderPath);
+		FString FullPath = FPaths::Combine(FOptPath, RiderPath);
 		TOptional<FInstallInfo> InstallInfo = FRiderPathLocator::GetInstallInfoFromRiderPath(FullPath, FInstallInfo::EInstallType::Installed);
 		if(InstallInfo.IsSet())
 		{
diff --git a/Engine/Plugins/Developer/RiderSourceCodeAccess/Source/RiderSourceCodeAccess/Private/RiderPathLocator/RiderPathLocator.h b/Engine/Plugins/Developer/RiderSourceCodeAccess/Source/RiderSourceCodeAccess/Private/RiderPathLocator/RiderPathLocator.h
--- a/Engine/Plugins/Developer/RiderSourceCodeAccess/Source/RiderSourceCodeAccess/Private/RiderPathLocator/RiderPathLocator.h
+++ b/Engine/Plugins/Developer/RiderSourceCodeAccess/Source/RiderSourceCodeAccess/Private/RiderPathLocator/RiderPathLocator.h
@@ -6,12 +6,7 @@
 #include "Misc/Paths.h"
 
 struct FVersion
-{
-	FVersion() = default;
-	~FVersion() = default;
-	
-	TArray<int32> Versions;
-
+{	
 	bool IsInitialized() const { return Versions.Num() != 0; }
 
 	void operator=(const FString& VersionString)
@@ -24,6 +19,33 @@ struct FVersion
 		}		
 	}
 
+	int32 Major() const
+	{
+		if(Versions.Num() >= 1)
+		{
+			return Versions[0];
+		}
+		return INVALID_VERSION;
+	}
+
+	int32 Minor() const
+	{
+		if(Versions.Num() >= 2)
+		{
+			return Versions[1];
+		}
+		return INVALID_VERSION;
+	}
+
+	int32 Patch() const
+	{
+		if(Versions.Num() >= 3)
+		{
+			return Versions[2];
+		}
+		return INVALID_VERSION;
+	}
+
 	bool operator<(const FVersion& rhs) const
 	{
 		TArray<int32>::SizeType LeftSize = Versions.Num();
@@ -44,6 +66,10 @@ struct FVersion
 	{
 		return FString::JoinBy(Versions, TEXT("."), [](int32 Item)-> FString { return FString::FromInt(Item); });
 	}
+
+	
+	TArray<int32> Versions;
+	static const int32 INVALID_VERSION = -1;
 };
 
 struct FInstallInfo
diff --git a/Engine/Plugins/Developer/RiderSourceCodeAccess/Source/RiderSourceCodeAccess/Private/RiderPathLocator/Win/RiderPathLocatorWin.cpp b/Engine/Plugins/Developer/RiderSourceCodeAccess/Source/RiderSourceCodeAccess/Private/RiderPathLocator/Win/RiderPathLocatorWin.cpp
--- a/Engine/Plugins/Developer/RiderSourceCodeAccess/Source/RiderSourceCodeAccess/Private/RiderPathLocator/Win/RiderPathLocatorWin.cpp
+++ b/Engine/Plugins/Developer/RiderSourceCodeAccess/Source/RiderSourceCodeAccess/Private/RiderPathLocator/Win/RiderPathLocatorWin.cpp
@@ -133,7 +133,9 @@ static TArray<FInstallInfo> CollectPathsFromRegistry( const Windows::HKEY RootKe
 					const FString ExePath = FPaths::Combine(InstallLocation, TEXT("bin"), TEXT("rider64.exe"));
 					TOptional<FInstallInfo> InstallInfo = FRiderPathLocator::GetInstallInfoFromRiderPath(ExePath, FInstallInfo::EInstallType::Installed);
 					if(InstallInfo.IsSet())
+					{
 						InstallInfos.Add(InstallInfo.GetValue());
+					}
 				}
 
 			}
@@ -145,16 +147,25 @@ static TArray<FInstallInfo> CollectPathsFromRegistry( const Windows::HKEY RootKe
 
 TOptional<FInstallInfo> FRiderPathLocator::GetInstallInfoFromRiderPath(const FString& Path, FInstallInfo::EInstallType InstallType)
 {
-	if(!FPaths::FileExists(Path)) return {};
+	if(!FPaths::FileExists(Path))
+	{
+		return {};
+	}
 	
 	const FString PatternString(TEXT("(.*)(?:\\\\|/)bin"));
 	const FRegexPattern Pattern(PatternString);
 	FRegexMatcher RiderPathMatcher(Pattern, Path);
-	if (!RiderPathMatcher.FindNext()) return {};
+	if (!RiderPathMatcher.FindNext())
+	{
+		return {};
+	}
 
 	const FString RiderDir = RiderPathMatcher.GetCaptureGroup(1);
 	const FString RiderCppPluginPath = FPaths::Combine(RiderDir, TEXT("plugins"), TEXT("rider-cpp"));
-	if (!FPaths::DirectoryExists(RiderCppPluginPath)) return {};
+	if (!FPaths::DirectoryExists(RiderCppPluginPath))
+	{
+		return {};
+	}
 	
 	FInstallInfo Info;
 	Info.Path = Path;
@@ -167,6 +178,10 @@ TOptional<FInstallInfo> FRiderPathLocator::GetInstallInfoFromRiderPath(const FSt
 	if(!Info.Version.IsInitialized())
 	{
 		Info.Version = FPaths::GetBaseFilename(RiderDir);
+		if(Info.Version.Major() >= 221)
+		{
+			Info.SupportUprojectState = FInstallInfo::ESupportUproject::Release;
+		}
 	}
 	return Info;
 }
diff --git a/Engine/Plugins/Developer/RiderSourceCodeAccess/Source/RiderSourceCodeAccess/Private/RiderSourceCodeAccessorModule.cpp b/Engine/Plugins/Developer/RiderSourceCodeAccess/Source/RiderSourceCodeAccess/Private/RiderSourceCodeAccessorModule.cpp
--- a/Engine/Plugins/Developer/RiderSourceCodeAccess/Source/RiderSourceCodeAccess/Private/RiderSourceCodeAccessorModule.cpp
+++ b/Engine/Plugins/Developer/RiderSourceCodeAccess/Source/RiderSourceCodeAccess/Private/RiderSourceCodeAccessorModule.cpp
@@ -39,12 +39,16 @@ void FRiderSourceCodeAccessModule::GenerateSlnAccessors(const TArray<FInstallInf
 #if PLATFORM_WINDOWS
 	if(InstallInfos.Num() == 0) return;
 	
-	for (const FInstallInfo& InstallInfo : InstallInfos)
+	if(InstallInfos.Num() > 1)
 	{
-		TSharedRef<FRiderSourceCodeAccessor> RiderSourceCodeAccessor = MakeShareable(new FRiderSourceCodeAccessor());
-		RiderSourceCodeAccessor->Init(InstallInfo, FRiderSourceCodeAccessor::EProjectModel::Sln);
-		IModularFeatures::Get().RegisterModularFeature(FRiderSourceCodeAccessor::FeatureType(), &RiderSourceCodeAccessor.Get());
-		RiderSourceCodeAccessors.Add(RiderSourceCodeAccessor->GetFName(), RiderSourceCodeAccessor);
+		for (const FInstallInfo& InstallInfo : InstallInfos)
+		{
+			TSharedRef<FRiderSourceCodeAccessor> RiderSourceCodeAccessor = MakeShareable(new FRiderSourceCodeAccessor());
+			RiderSourceCodeAccessor->Init(InstallInfo, FRiderSourceCodeAccessor::EProjectModel::Sln);
+			IModularFeatures::Get().RegisterModularFeature(FRiderSourceCodeAccessor::FeatureType(), &RiderSourceCodeAccessor.Get());
+			RiderSourceCodeAccessors.Add(RiderSourceCodeAccessor->GetFName(), RiderSourceCodeAccessor);
+		}
+
 	}
 
 	const TSharedRef<FRiderSourceCodeAccessor> RiderSourceCodeAccessor = MakeShareable(new FRiderSourceCodeAccessor());
@@ -64,12 +68,15 @@ void FRiderSourceCodeAccessModule::GenerateUprojectAccessors(const TArray<FInsta
 
 	if(UprojectInfos.Num() == 0) return;
 
-	for (const FInstallInfo& UprojectInfo : UprojectInfos)
+	if(UprojectInfos.Num() > 1)
 	{
-		TSharedRef<FRiderSourceCodeAccessor> RiderSourceCodeAccessor = MakeShareable(new FRiderSourceCodeAccessor());
-		RiderSourceCodeAccessor->Init(UprojectInfo, FRiderSourceCodeAccessor::EProjectModel::Uproject);
-		IModularFeatures::Get().RegisterModularFeature(FRiderSourceCodeAccessor::FeatureType(), &RiderSourceCodeAccessor.Get());
-		RiderSourceCodeAccessors.Add(RiderSourceCodeAccessor->GetFName(), RiderSourceCodeAccessor);
+		for (const FInstallInfo& UprojectInfo : UprojectInfos)
+		{
+			TSharedRef<FRiderSourceCodeAccessor> RiderSourceCodeAccessor = MakeShareable(new FRiderSourceCodeAccessor());
+			RiderSourceCodeAccessor->Init(UprojectInfo, FRiderSourceCodeAccessor::EProjectModel::Uproject);
+			IModularFeatures::Get().RegisterModularFeature(FRiderSourceCodeAccessor::FeatureType(), &RiderSourceCodeAccessor.Get());
+			RiderSourceCodeAccessors.Add(RiderSourceCodeAccessor->GetFName(), RiderSourceCodeAccessor);
+		}
 	}
 
 	const TSharedRef<FRiderSourceCodeAccessor> RiderSourceCodeAccessor = MakeShareable(new FRiderSourceCodeAccessor());
