Fix compiler failure if compiling standalone cpp file

diff --git a/Engine/Plugins/Developer/RiderSourceCodeAccess/Source/RiderSourceCodeAccess/Private/RiderPathLocator/Linux/RiderPathLocatorLinux.cpp b/Engine/Plugins/Developer/RiderSourceCodeAccess/Source/RiderSourceCodeAccess/Private/RiderPathLocator/Linux/RiderPathLocatorLinux.cpp
--- a/Engine/Plugins/Developer/RiderSourceCodeAccess/Source/RiderSourceCodeAccess/Private/RiderPathLocator/Linux/RiderPathLocatorLinux.cpp
+++ b/Engine/Plugins/Developer/RiderSourceCodeAccess/Source/RiderSourceCodeAccess/Private/RiderPathLocator/Linux/RiderPathLocatorLinux.cpp
@@ -2,6 +2,7 @@
 
 #include "RiderPathLocator/RiderPathLocator.h"
 
+#include "Internationalization/Regex.h"
 #include "HAL/FileManager.h"
 #include "Misc/FileHelper.h"
 #include "Misc/Paths.h"
