Fix crash on startup due to Rider adding multiple of the same names due to an * in the strings

diff --git a/Engine/Plugins/Developer/RiderSourceCodeAccess/Source/RiderSourceCodeAccess/Private/RiderPathLocator/Linux/RiderPathLocatorLinux.cpp b/Engine/Plugins/Developer/RiderSourceCodeAccess/Source/RiderSourceCodeAccess/Private/RiderPathLocator/Linux/RiderPathLocatorLinux.cpp
--- a/Engine/Plugins/Developer/RiderSourceCodeAccess/Source/RiderSourceCodeAccess/Private/RiderPathLocator/Linux/RiderPathLocatorLinux.cpp
+++ b/Engine/Plugins/Developer/RiderSourceCodeAccess/Source/RiderSourceCodeAccess/Private/RiderPathLocator/Linux/RiderPathLocatorLinux.cpp
@@ -72,8 +72,8 @@ static TArray<FInstallInfo> GetManuallyInstalledRiders()
 	TArray<FString> RiderPaths;
 
 	const FString FHomePath = GetHomePath();
-	const FString HomePathMask = FPaths::Combine(FHomePath, TEXT("Rider*"));
-	
+	const FString HomePathMask = FPaths::Combine(FHomePath, TEXT("Rider.sh"));
+
 	IFileManager::Get().FindFiles(RiderPaths, *HomePathMask, false, true);
 
 	for(const FString& RiderPath: RiderPaths)
@@ -87,8 +87,8 @@ static TArray<FInstallInfo> GetManuallyInstalledRiders()
 	}
 
 	const FString FOptPath = TEXT("/opt");
-	const FString OptPathMask = FPaths::Combine(FOptPath, TEXT("Rider*"));
-	
+	const FString OptPathMask = FPaths::Combine(FOptPath, TEXT("Rider.sh"));
+
 	IFileManager::Get().FindFiles(RiderPaths, *OptPathMask, false, true);
 
 	for(const FString& RiderPath: RiderPaths)
@@ -155,7 +155,7 @@ TSet<FInstallInfo> FRiderPathLocator::CollectAllPaths()
 	TSet<FInstallInfo> InstallInfos;
 	InstallInfos.Append(GetInstalledRidersWithMdfind());
 	InstallInfos.Append(GetManuallyInstalledRiders());
-	InstallInfos.Append(GetInstallInfosFromToolbox(GetToolboxPath(), "Rider*"));
+	InstallInfos.Append(GetInstallInfosFromToolbox(GetToolboxPath(), "Rider.sh"));
 	InstallInfos.Append(GetInstallInfosFromResourceFile());
 	return InstallInfos;
 }
