Fix missing header

diff --git a/Engine/Plugins/Developer/RiderSourceCodeAccess/Source/RiderSourceCodeAccess/Private/RiderPathLocator/Linux/RiderPathLocatorLinux.cpp b/Engine/Plugins/Developer/RiderSourceCodeAccess/Source/RiderSourceCodeAccess/Private/RiderPathLocator/Linux/RiderPathLocatorLinux.cpp
--- a/Engine/Plugins/Developer/RiderSourceCodeAccess/Source/RiderSourceCodeAccess/Private/RiderPathLocator/Linux/RiderPathLocatorLinux.cpp
+++ b/Engine/Plugins/Developer/RiderSourceCodeAccess/Source/RiderSourceCodeAccess/Private/RiderPathLocator/Linux/RiderPathLocatorLinux.cpp
@@ -4,6 +4,7 @@
 
 #include "Internationalization/Regex.h"
 #include "HAL/FileManager.h"
+#include "HAL/PlatformProcess.h"
 #include "Misc/FileHelper.h"
 #include "Misc/Paths.h"
 #include "Serialization/JsonSerializer.h"
diff --git a/Engine/Plugins/Developer/RiderSourceCodeAccess/Source/RiderSourceCodeAccess/Private/RiderPathLocator/Mac/RiderPathLocatorMac.cpp b/Engine/Plugins/Developer/RiderSourceCodeAccess/Source/RiderSourceCodeAccess/Private/RiderPathLocator/Mac/RiderPathLocatorMac.cpp
--- a/Engine/Plugins/Developer/RiderSourceCodeAccess/Source/RiderSourceCodeAccess/Private/RiderPathLocator/Mac/RiderPathLocatorMac.cpp
+++ b/Engine/Plugins/Developer/RiderSourceCodeAccess/Source/RiderSourceCodeAccess/Private/RiderPathLocator/Mac/RiderPathLocatorMac.cpp
@@ -3,6 +3,7 @@
 #include "RiderPathLocator/RiderPathLocator.h"
 
 #include "HAL/FileManager.h"
+#include "HAL/PlatformProcess.h"
 #include "Misc/FileHelper.h"
 #include "Misc/Paths.h"
 #include "Serialization/JsonSerializer.h"
