Add ShortNames to Code Access plugins to reduce the pressure on path length. Problem reported on UDN.

diff --git a/Engine/Plugins/Developer/CLionSourceCodeAccess/Source/CLionSourceCodeAccess/CLionSourceCodeAccess.Build.cs b/Engine/Plugins/Developer/CLionSourceCodeAccess/Source/CLionSourceCodeAccess/CLionSourceCodeAccess.Build.cs
--- a/Engine/Plugins/Developer/CLionSourceCodeAccess/Source/CLionSourceCodeAccess/CLionSourceCodeAccess.Build.cs
+++ b/Engine/Plugins/Developer/CLionSourceCodeAccess/Source/CLionSourceCodeAccess/CLionSourceCodeAccess.Build.cs
@@ -20,6 +20,8 @@ namespace UnrealBuildTool.Rules
             {
                 PrivateDependencyModuleNames.Add("HotReload");
             }
+
+			ShortName = "CLionSCA";
 		}
 	}
 }
\ No newline at end of file
diff --git a/Engine/Plugins/Developer/CodeLiteSourceCodeAccess/Source/CodeLiteSourceCodeAccess/CodeLiteSourceCodeAccess.Build.cs b/Engine/Plugins/Developer/CodeLiteSourceCodeAccess/Source/CodeLiteSourceCodeAccess/CodeLiteSourceCodeAccess.Build.cs
--- a/Engine/Plugins/Developer/CodeLiteSourceCodeAccess/Source/CodeLiteSourceCodeAccess/CodeLiteSourceCodeAccess.Build.cs
+++ b/Engine/Plugins/Developer/CodeLiteSourceCodeAccess/Source/CodeLiteSourceCodeAccess/CodeLiteSourceCodeAccess.Build.cs
@@ -19,6 +19,8 @@ namespace UnrealBuildTool.Rules
 			{
 				PrivateDependencyModuleNames.Add("HotReload");
 			}
+
+			ShortName = "CLiteSCA";
 		}
 	}
 }
diff --git a/Engine/Plugins/Developer/KDevelopSourceCodeAccess/Source/KDevelopSourceCodeAccess/KDevelopSourceCodeAccess.Build.cs b/Engine/Plugins/Developer/KDevelopSourceCodeAccess/Source/KDevelopSourceCodeAccess/KDevelopSourceCodeAccess.Build.cs
--- a/Engine/Plugins/Developer/KDevelopSourceCodeAccess/Source/KDevelopSourceCodeAccess/KDevelopSourceCodeAccess.Build.cs
+++ b/Engine/Plugins/Developer/KDevelopSourceCodeAccess/Source/KDevelopSourceCodeAccess/KDevelopSourceCodeAccess.Build.cs
@@ -19,6 +19,8 @@ namespace UnrealBuildTool.Rules
 			{
 				PrivateDependencyModuleNames.Add("HotReload");
 			}
+
+			ShortName = "KDevSCA";
 		}
 	}
 }
diff --git a/Engine/Plugins/Developer/NullSourceCodeAccess/Source/NullSourceCodeAccess/NullSourceCodeAccess.Build.cs b/Engine/Plugins/Developer/NullSourceCodeAccess/Source/NullSourceCodeAccess/NullSourceCodeAccess.Build.cs
--- a/Engine/Plugins/Developer/NullSourceCodeAccess/Source/NullSourceCodeAccess/NullSourceCodeAccess.Build.cs
+++ b/Engine/Plugins/Developer/NullSourceCodeAccess/Source/NullSourceCodeAccess/NullSourceCodeAccess.Build.cs
@@ -14,6 +14,8 @@ namespace UnrealBuildTool.Rules
 					"SourceCodeAccess"
 				}
 			);
+
+			ShortName = "NullSCA";
 		}
 	}
 }
diff --git a/Engine/Plugins/Developer/RiderSourceCodeAccess/Source/RiderSourceCodeAccess/RiderSourceCodeAccess.Build.cs b/Engine/Plugins/Developer/RiderSourceCodeAccess/Source/RiderSourceCodeAccess/RiderSourceCodeAccess.Build.cs
--- a/Engine/Plugins/Developer/RiderSourceCodeAccess/Source/RiderSourceCodeAccess/RiderSourceCodeAccess.Build.cs
+++ b/Engine/Plugins/Developer/RiderSourceCodeAccess/Source/RiderSourceCodeAccess/RiderSourceCodeAccess.Build.cs
@@ -1,4 +1,4 @@
-﻿// Copyright Epic Games, Inc. All Rights Reserved.
+// Copyright Epic Games, Inc. All Rights Reserved.
 
 namespace UnrealBuildTool.Rules
 {
@@ -25,6 +25,8 @@ namespace UnrealBuildTool.Rules
 				PrivateDependencyModuleNames.Add("UnrealEd");
 				PrivateDependencyModuleNames.Add("GameProjectGeneration");
 			}
+
+			ShortName = "RiderSCA";
 		}
 	}
 }
\ No newline at end of file
diff --git a/Engine/Plugins/Developer/VisualStudioCodeSourceCodeAccess/Source/VisualStudioCodeSourceCodeAccess/VisualStudioCodeSourceCodeAccess.Build.cs b/Engine/Plugins/Developer/VisualStudioCodeSourceCodeAccess/Source/VisualStudioCodeSourceCodeAccess/VisualStudioCodeSourceCodeAccess.Build.cs
--- a/Engine/Plugins/Developer/VisualStudioCodeSourceCodeAccess/Source/VisualStudioCodeSourceCodeAccess/VisualStudioCodeSourceCodeAccess.Build.cs
+++ b/Engine/Plugins/Developer/VisualStudioCodeSourceCodeAccess/Source/VisualStudioCodeSourceCodeAccess/VisualStudioCodeSourceCodeAccess.Build.cs
@@ -56,6 +56,8 @@ namespace UnrealBuildTool.Rules
 			}
 
 			bBuildLocallyWithSNDBS = true;
+
+			ShortName = "VSCSCA";
 		}
 	}
 }
diff --git a/Engine/Plugins/Developer/VisualStudioSourceCodeAccess/Source/VisualStudioSourceCodeAccess/VisualStudioSourceCodeAccess.Build.cs b/Engine/Plugins/Developer/VisualStudioSourceCodeAccess/Source/VisualStudioSourceCodeAccess/VisualStudioSourceCodeAccess.Build.cs
--- a/Engine/Plugins/Developer/VisualStudioSourceCodeAccess/Source/VisualStudioSourceCodeAccess/VisualStudioSourceCodeAccess.Build.cs
+++ b/Engine/Plugins/Developer/VisualStudioSourceCodeAccess/Source/VisualStudioSourceCodeAccess/VisualStudioSourceCodeAccess.Build.cs
@@ -27,6 +27,8 @@ namespace UnrealBuildTool.Rules
 			}
 
 			bBuildLocallyWithSNDBS = true;
+			
+			ShortName = "VSSCA";
 		}
 	}
 }
diff --git a/Engine/Plugins/Developer/XCodeSourceCodeAccess/Source/XCodeSourceCodeAccess/XCodeSourceCodeAccess.Build.cs b/Engine/Plugins/Developer/XCodeSourceCodeAccess/Source/XCodeSourceCodeAccess/XCodeSourceCodeAccess.Build.cs
--- a/Engine/Plugins/Developer/XCodeSourceCodeAccess/Source/XCodeSourceCodeAccess/XCodeSourceCodeAccess.Build.cs
+++ b/Engine/Plugins/Developer/XCodeSourceCodeAccess/Source/XCodeSourceCodeAccess/XCodeSourceCodeAccess.Build.cs
@@ -14,6 +14,8 @@ namespace UnrealBuildTool.Rules
 					"DesktopPlatform"
 				}
 			);
+
+			ShortName = "XCodeSCA";
 		}
 	}
 }
