PR #9111: Fix opening Rider on Mac M1 (Contributed by [REDACTED_PERSON])

#rnx

diff --git a/Engine/Plugins/Developer/RiderSourceCodeAccess/Source/RiderSourceCodeAccess/Private/RiderSourceCodeAccessor.cpp b/Engine/Plugins/Developer/RiderSourceCodeAccess/Source/RiderSourceCodeAccess/Private/RiderSourceCodeAccessor.cpp
--- a/Engine/Plugins/Developer/RiderSourceCodeAccess/Source/RiderSourceCodeAccess/Private/RiderSourceCodeAccessor.cpp
+++ b/Engine/Plugins/Developer/RiderSourceCodeAccess/Source/RiderSourceCodeAccess/Private/RiderSourceCodeAccessor.cpp
@@ -16,14 +16,17 @@
 #include "ProjectDescriptor.h"
 #include "GameProjectGenerationModule.h"
 #include "Dialogs/SOutputLogDialog.h"
+#include "Framework/Notifications/NotificationManager.h"
 #include "Misc/MessageDialog.h"
+#include "Widgets/Notifications/SNotificationList.h"
 
 #define LOCTEXT_NAMESPACE "RiderSourceCodeAccessor"
 
 DEFINE_LOG_CATEGORY_STATIC(LogRiderAccessor, Log, All);
 
-namespace
+namespace RSCA
 {
+
 TOptional<FString> ResolvePathToFile(const FString& FullPath)
 {
 	FString Path = FullPath;
@@ -52,6 +55,84 @@ TOptional<FString> ResolvePathToFile(const FString& FullPath)
 	}
 	return {Path};
 }
+
+struct FCommandLineInfo
+{
+	FString App;
+	FString Args;
+};
+
+#if PLATFORM_MAC
+static int ProcessIsTranslated()
+{
+	int Return = 0;
+
+	size_t Size = sizeof(Return);
+	if (sysctlbyname("sysctl.proc_translated", &Return, &Size, NULL, 0) == -1)
+	{
+		if (errno == ENOENT)
+		{
+			return 0;
+		}
+		return -1;
+	}
+
+	return Return;
+}
+#endif //PLATFORM_MAC
+
+FCommandLineInfo GetPlatformAppAndArgs(const FString& App, const FString& Args)
+{
+	FCommandLineInfo info;
+	info.App = App;
+    info.Args = Args;
+
+#if PLATFORM_MAC
+	if (ProcessIsTranslated() == 1)
+	{
+		info.App = TEXT("/usr/bin/arch");
+		info.Args = FString::Printf(TEXT("-arm64 \"%s\" %s"), *App, *Args);
+	}
+#endif // PLATFORM_MAC
+
+	return info;
+}
+
+bool CheckExecutable(const FString& App)
+{
+	if (FPaths::FileExists(App) || FPaths::DirectoryExists(App))
+	{
+		return true;
+	}
+
+	FNotificationInfo Info(FText::Format(LOCTEXT("CodeAccessorAppDoesntExist", "{0} doesn't exist"), FText::FromString(App)));
+	Info.bFireAndForget = true;
+
+	FSlateNotificationManager::Get().AddNotification(Info)->SetCompletionState(SNotificationItem::CS_Fail);
+
+	return false;
+}
+
+bool OpenRider(FString const& ExecutablePath, FString const& Params, FString const& ErrorMessage)
+{
+	const FCommandLineInfo PlatformAppAndArgs = GetPlatformAppAndArgs(ExecutablePath, Params);
+
+	if (!CheckExecutable(PlatformAppAndArgs.App))
+	{
+		return false;
+	}
+
+	FProcHandle Proc = FPlatformProcess::CreateProc(*PlatformAppAndArgs.App, *PlatformAppAndArgs.Args, true, true, false, nullptr, 0, nullptr, nullptr);
+	const bool bResult = Proc.IsValid();
+	if (!bResult)
+	{
+		UE_LOG(LogRiderAccessor, Warning, TEXT("%s"), *ErrorMessage);
+		FPlatformProcess::CloseProc(Proc);
+	}
+
+	return bResult;
+}
+
 }
 
 void FRiderSourceCodeAccessor::RefreshAvailability()
@@ -105,23 +186,16 @@ bool FRiderSourceCodeAccessor::OpenFileAtLine(const FString& FullPath, int32 Lin
 	if (FPaths::IsRelative(SolutionPath))
 		SolutionPath = FPaths::ConvertRelativePathToFull(SolutionPath);
 
-	const TOptional<FString> OptionalPath = ResolvePathToFile(FullPath);
+	const TOptional<FString> OptionalPath = RSCA::ResolvePathToFile(FullPath);
 	if(!OptionalPath.IsSet()) return false;
 	
 	const FString Path = OptionalPath.GetValue();
 	const FString Params = FString::Printf(TEXT("\"%s\" --line %d \"%s\""), *SolutionPath, LineNumber, *Path);
+	const FString ErrorMessage = FString::Printf(TEXT("Opening file (%s) at a line (%d) failed."), *Path, LineNumber);
 
-	return HandleOpeningRider([this, &Params, &Path, LineNumber]()	->bool
+	return HandleOpeningRider([this, &Params, &ErrorMessage]() -> bool
 	{
-		FProcHandle Proc = FPlatformProcess::CreateProc(*ExecutablePath, *Params, true, true, false, nullptr, 0,
-		                                                nullptr, nullptr);
-		const bool bResult = Proc.IsValid();
-		if (!bResult)
-		{
-			UE_LOG(LogRiderAccessor, Warning, TEXT("Opening file (%s) at a line (%d) failed."), *Path, LineNumber);
-			FPlatformProcess::CloseProc(Proc);
-		}
-		return bResult;
+		return RSCA::OpenRider(ExecutablePath, Params, ErrorMessage);
 	});
 }
 
@@ -136,22 +210,13 @@ bool FRiderSourceCodeAccessor::OpenSolution()
 	
 	const FString FullPath = IFileManager::Get().ConvertToAbsolutePathForExternalAppForRead(*SolutionPath);
 	const FString Params = FString::Printf(TEXT("\"%s\""), *FullPath);
+	const FString ErrorMessage = FString::Printf(TEXT("Opening solution (%s) failed."), *FullPath);
 
-	return HandleOpeningRider([this, &Params, &FullPath]()->bool
+	return HandleOpeningRider([this, &Params, &ErrorMessage]()->bool
 	{
-		FProcHandle Proc = FPlatformProcess::CreateProc(*ExecutablePath, *Params, true, true, false, nullptr, 0,
-		                                                nullptr, nullptr);
-		const bool bResult = Proc.IsValid();
-		if (!bResult)
-		{
-			UE_LOG(LogRiderAccessor, Warning, TEXT("Opening solution (%s) failed."), *FullPath);
-			FPlatformProcess::CloseProc(Proc);
-		}
-
-		return bResult;
+		return RSCA::OpenRider(ExecutablePath, Params, ErrorMessage);
 	});
 }
-
 bool FRiderSourceCodeAccessor::OpenSolutionAtPath(const FString& InSolutionPath)
 {
 	if (!bHasRiderInstalled) return false;
@@ -162,17 +227,11 @@ bool FRiderSourceCodeAccessor::OpenSolutionAtPath(const FString& InSolutionPath)
 		CorrectSolutionPath += ".sln";
 	}
 	const FString Params = FString::Printf(TEXT("\"%s\""), *CorrectSolutionPath);
-	
-	return HandleOpeningRider([this, &Params, &CorrectSolutionPath]()->bool
+	const FString ErrorMessage = FString::Printf(TEXT("Opening the project file (%s) failed."), *CorrectSolutionPath);
+
+	return HandleOpeningRider([this, &Params, &ErrorMessage]()->bool
 	{
-	    FProcHandle Proc = FPlatformProcess::CreateProc(*ExecutablePath, *Params, true, true, false, nullptr, 0, nullptr, nullptr);
-	    const bool bResult = Proc.IsValid(); 
-	    if (!bResult)
-	    {
-	        UE_LOG(LogRiderAccessor, Warning, TEXT("Opening the project file (%s) failed."), *CorrectSolutionPath);
-	        FPlatformProcess::CloseProc(Proc);
-	    }
-		return bResult;
+		return RSCA::OpenRider(ExecutablePath, Params, ErrorMessage);
 	});
 }
 
@@ -199,24 +258,18 @@ bool FRiderSourceCodeAccessor::OpenSourceFiles(const TArray<FString>& AbsoluteSo
 
 	FString FilePaths = "";
 	for (const FString & FullPath : AbsoluteSourcePaths) {
-		const TOptional<FString> OptionalPath = ResolvePathToFile(FullPath);
+		const TOptional<FString> OptionalPath = RSCA::ResolvePathToFile(FullPath);
 		if(!OptionalPath.IsSet()) return false;
 		const FString Path = OptionalPath.GetValue();
 		FilePaths += FString::Printf(TEXT("\"%s\" "), *Path);
 	}
 
 	const FString Params = FString::Printf(TEXT("\"%s\" %s"), *SolutionPath, *FilePaths);
+	const FString ErrorMessage = FString::Printf(TEXT("Opening files (%s) failed."), *FilePaths);
 
-	return HandleOpeningRider([this, &Params, &FilePaths]()->bool
+	return HandleOpeningRider([this, &Params, &ErrorMessage]()->bool
 	{
-	    FProcHandle Proc = FPlatformProcess::CreateProc(*ExecutablePath, *Params, true, true, false, nullptr, 0, nullptr, nullptr);
-	    const bool bResult = Proc.IsValid(); 
-	    if (!bResult)
-	    {
-	        UE_LOG(LogRiderAccessor, Warning, TEXT("Opening files (%s) failed."), *FilePaths);
-	        FPlatformProcess::CloseProc(Proc);
-	    }
-	    return bResult;		
+		return RSCA::OpenRider(ExecutablePath, Params, ErrorMessage);
 	});
 }
 
diff --git a/Engine/Plugins/Developer/RiderSourceCodeAccess/Source/RiderSourceCodeAccess/RiderSourceCodeAccess.Build.cs b/Engine/Plugins/Developer/RiderSourceCodeAccess/Source/RiderSourceCodeAccess/RiderSourceCodeAccess.Build.cs
--- a/Engine/Plugins/Developer/RiderSourceCodeAccess/Source/RiderSourceCodeAccess/RiderSourceCodeAccess.Build.cs
+++ b/Engine/Plugins/Developer/RiderSourceCodeAccess/Source/RiderSourceCodeAccess/RiderSourceCodeAccess.Build.cs
@@ -15,7 +15,8 @@ namespace UnrealBuildTool.Rules
                     "DesktopPlatform",
 					"Projects",
 					"Json",
-					"Projects"
+					"Projects",
+					"Slate"
 				}
 			);
 
