PR #9371: Fix access violation (Contributed by [REDACTED_PERSON])

diff --git a/Engine/Plugins/Developer/RiderSourceCodeAccess/Source/RiderSourceCodeAccess/Private/RiderSourceCodeAccessorModule.h b/Engine/Plugins/Developer/RiderSourceCodeAccess/Source/RiderSourceCodeAccess/Private/RiderSourceCodeAccessorModule.h
--- a/Engine/Plugins/Developer/RiderSourceCodeAccess/Source/RiderSourceCodeAccess/Private/RiderSourceCodeAccessorModule.h
+++ b/Engine/Plugins/Developer/RiderSourceCodeAccess/Source/RiderSourceCodeAccess/Private/RiderSourceCodeAccessorModule.h
@@ -14,5 +14,5 @@ public:
 private:
 	void GenerateSlnAccessors(const TArray<struct FInstallInfo>& InstallInfos);
 	void GenerateUprojectAccessors(const TArray<struct FInstallInfo>& InstallInfos);
-	TMap<FName, TSharedRef<ISourceCodeAccessor>> RiderSourceCodeAccessors;
+	TMultiMap<FName, TSharedRef<ISourceCodeAccessor>> RiderSourceCodeAccessors;
 };
