UE on Linux now locates any standalone Rider installation if either installed to the default location "/opt" (https://www.jetbrains.com/help/rider/Installation_guide.html#fac067cf), or anywhere else that is on the PATH.

diff --git a/Engine/Plugins/Developer/RiderSourceCodeAccess/Source/RiderSourceCodeAccess/Private/RiderPathLocator/Linux/RiderPathLocatorLinux.cpp b/Engine/Plugins/Developer/RiderSourceCodeAccess/Source/RiderSourceCodeAccess/Private/RiderPathLocator/Linux/RiderPathLocatorLinux.cpp
--- a/Engine/Plugins/Developer/RiderSourceCodeAccess/Source/RiderSourceCodeAccess/Private/RiderPathLocator/Linux/RiderPathLocatorLinux.cpp
+++ b/Engine/Plugins/Developer/RiderSourceCodeAccess/Source/RiderSourceCodeAccess/Private/RiderPathLocator/Linux/RiderPathLocatorLinux.cpp
@@ -19,7 +19,7 @@ TOptional<FInstallInfo> FRiderPathLocator::GetInstallInfoFromRiderPath(const FSt
 	{
 		return {};
 	}
-	
+
 	const FString PatternString(TEXT("(.*)(?:\\\\|/)bin"));
 	const FRegexPattern Pattern(PatternString);
 	FRegexMatcher RiderPathMatcher(Pattern, Path);
@@ -34,7 +34,7 @@ TOptional<FInstallInfo> FRiderPathLocator::GetInstallInfoFromRiderPath(const FSt
 	{
 		return {};
 	}
-	
+
 	FInstallInfo Info;
 	Info.Path = Path;
 	Info.InstallType = InstallType;
@@ -67,41 +67,60 @@ static FString GetHomePath()
 	return FHomePath;
 }
 
-static TArray<FInstallInfo> GetManuallyInstalledRiders()
+// Installing the standalone version of Rider does not set up path variables, and it puts Rider in /opt
+static TArray<FInstallInfo> GetStandaloneInstalledRiders()
 {
-	TArray<FInstallInfo> Result;
-	TArray<FString> RiderPaths;
+	TArray<FInstallInfo> Results;
 
-	const FString FHomePath = GetHomePath();
-	const FString HomePathMask = FPaths::Combine(FHomePath, TEXT("Rider.sh"));
+	const FString OptPath = TEXT("/opt");
+	const FString OptPathMask = FPaths::Combine(OptPath, TEXT("*rider*"));
 
-	IFileManager::Get().FindFiles(RiderPaths, *HomePathMask, false, true);
+	TArray<FString> FolderPaths;
+	IFileManager::Get().FindFiles(FolderPaths, *OptPathMask, false, true);
 
-	for(const FString& RiderPath: RiderPaths)
+	for (const FString& RiderPath : FolderPaths)
 	{
-		FString FullPath = FPaths::Combine(FHomePath, RiderPath);
-		TOptional<FInstallInfo> InstallInfo = FRiderPathLocator::GetInstallInfoFromRiderPath(FullPath, FInstallInfo::EInstallType::Installed);
-		if(InstallInfo.IsSet())
+		const FString FolderPath = FPaths::Combine(OptPath, RiderPath);
+
+		TArray<FString> FullPaths;
+		IFileManager::Get().FindFilesRecursive(FullPaths, *FolderPath, TEXT("rider.sh"), true, false, false);
+
+		for (const FString& FullPath : FullPaths)
 		{
-			Result.Add(InstallInfo.GetValue());
+			TOptional<FInstallInfo> InstallInfo = FRiderPathLocator::GetInstallInfoFromRiderPath(FullPath, FInstallInfo::EInstallType::Installed);
+			if(InstallInfo.IsSet())
+			{
+				Results.Add(InstallInfo.GetValue());
+			}
 		}
 	}
 
-	const FString FOptPath = TEXT("/opt");
-	const FString OptPathMask = FPaths::Combine(FOptPath, TEXT("Rider.sh"));
+	return Results;
+}
+
+// Find any Rider installations on the PATH
+static TArray<FInstallInfo> GetInstalledRidersFromPath()
+{
+	TArray<FInstallInfo> Results;
 
-	IFileManager::Get().FindFiles(RiderPaths, *OptPathMask, false, true);
+	// Search each folder on the PATH
+	TArray<FString> Paths;
+	FPlatformMisc::GetEnvironmentVariable(TEXT("PATH")).ParseIntoArray(Paths, FPlatformMisc::GetPathVarDelimiter());
 
-	for(const FString& RiderPath: RiderPaths)
+	for (const FString& Path : Paths)
 	{
-		FString FullPath = FPaths::Combine(FOptPath, RiderPath);
-		TOptional<FInstallInfo> InstallInfo = FRiderPathLocator::GetInstallInfoFromRiderPath(FullPath, FInstallInfo::EInstallType::Installed);
-		if(InstallInfo.IsSet())
+		FString RiderPath = FPaths::Combine(Path, TEXT("rider.sh"));
+		if (FPaths::FileExists(RiderPath))
 		{
-			Result.Add(InstallInfo.GetValue());
+			TOptional<FInstallInfo> InstallInfo = FRiderPathLocator::GetInstallInfoFromRiderPath(RiderPath, FInstallInfo::EInstallType::Installed);
+			if(InstallInfo.IsSet())
+			{
+				Results.Add(InstallInfo.GetValue());
+			}
 		}
 	}
-	return Result;
+
+	return Results;
 }
 
 static FString GetToolboxPath()
@@ -155,7 +174,8 @@ TSet<FInstallInfo> FRiderPathLocator::CollectAllPaths()
 {
 	TSet<FInstallInfo> InstallInfos;
 	InstallInfos.Append(GetInstalledRidersWithMdfind());
-	InstallInfos.Append(GetManuallyInstalledRiders());
+	InstallInfos.Append(GetStandaloneInstalledRiders());
+	InstallInfos.Append(GetInstalledRidersFromPath());
 	InstallInfos.Append(GetInstallInfosFromToolbox(GetToolboxPath(), "Rider.sh"));
 	InstallInfos.Append(GetInstallInfosFromResourceFile());
 	return InstallInfos;
