RiderSourceAccessor: When in uproject model, do not try to launch a sln when opening from Project Launcher

diff --git a/Engine/Plugins/Developer/RiderSourceCodeAccess/Source/RiderSourceCodeAccess/Private/RiderSourceCodeAccessor.cpp b/Engine/Plugins/Developer/RiderSourceCodeAccess/Source/RiderSourceCodeAccess/Private/RiderSourceCodeAccessor.cpp
--- a/Engine/Plugins/Developer/RiderSourceCodeAccess/Source/RiderSourceCodeAccess/Private/RiderSourceCodeAccessor.cpp
+++ b/Engine/Plugins/Developer/RiderSourceCodeAccess/Source/RiderSourceCodeAccess/Private/RiderSourceCodeAccessor.cpp
@@ -222,7 +222,11 @@ bool FRiderSourceCodeAccessor::OpenSolutionAtPath(const FString& InSolutionPath)
 	if (!bHasRiderInstalled) return false;
 
 	FString CorrectSolutionPath = InSolutionPath;
-	if (!CorrectSolutionPath.EndsWith(".sln"))
+	if (Model == EProjectModel::Uproject && !CorrectSolutionPath.EndsWith(".uproject"))
+	{
+		CorrectSolutionPath += ".uproject";
+	}
+	else if (!CorrectSolutionPath.EndsWith(".sln"))
 	{
 		CorrectSolutionPath += ".sln";
 	}
