PR #9811: Updating RiderSourceCodeAccess (Contributed by [REDACTED_PERSON])

diff --git a/Engine/Plugins/Developer/RiderSourceCodeAccess/README.md b/Engine/Plugins/Developer/RiderSourceCodeAccess/README.md
--- a/Engine/Plugins/Developer/RiderSourceCodeAccess/README.md
+++ b/Engine/Plugins/Developer/RiderSourceCodeAccess/README.md
@@ -16,7 +16,7 @@ RiderSourceCodeAccess is a plugin for Unreal Engine, available through [Marketpl
 
 # How to add this plugin manually (not recommened)
 ## How to add to Game project
-1. Go to ["Releases" page]([REDACTED_PROFILE_URL]/RiderSourceCodeAccess/releases) and download version of plugin for your version of Unreal Engine; 
+1. Get the source code of the [master branch]([REDACTED_PROFILE_URL]/RiderSourceCodeAccess/archive/refs/heads/master.zip); 
 2. Unzip `RiderSourceCodeAccess.zip` to `RiderSourceCodeAccess` folder;
 3. Copy `RiderSourceCodeAccess` folder to `{GameProjectRoot}/Plugins/Developer`;
   a. If `{GameProjectRoot}/Plugins/Developer` does not exist, create it;
@@ -28,7 +28,7 @@ RiderSourceCodeAccess is a plugin for Unreal Engine, available through [Marketpl
 
 ## How to add to Engine project
 ### Unreal Engine is built from source code   
-1. Go to ["Releases" page]([REDACTED_PROFILE_URL]/RiderSourceCodeAccess/releases) and download version of plugin for your version of Unreal Engine; 
+1. Get the source code of the [master branch]([REDACTED_PROFILE_URL]/RiderSourceCodeAccess/archive/refs/heads/master.zip); 
 2. Unzip `RiderSourceCodeAccess.zip` to `RiderSourceCodeAccess` folder;
 3. Copy `RiderSourceCodeAccess` folder to `{UnrealEngineRoot}/Engine/Plugins/Developer`;
   a. If `{UnrealEngineRoot}/Engine/Plugins/Developer` does not exist, create it;
@@ -39,7 +39,7 @@ RiderSourceCodeAccess is a plugin for Unreal Engine, available through [Marketpl
 8. Select Rider from drop down list. NB: only Rider with C++ plugins will be available from drop down list.
 
 ### Unreal Engine is downloaded from Epic Games Store
-1. Go to ["Releases" page]([REDACTED_PROFILE_URL]/RiderSourceCodeAccess/releases) and download version of plugin for your version of Unreal Engine; 
+1. Get the source code of the [master branch]([REDACTED_PROFILE_URL]/RiderSourceCodeAccess/archive/refs/heads/master.zip); 
 2. Unzip `RiderSourceCodeAccess.zip` to `RiderSourceCodeAccess` folder;
 3. Copy `RiderSourceCodeAccess` folder to `{UnrealProjectRoot}/Plugins/Developer`;
   a. If `{UnrealProjectRoot}/Plugins/Developer` does not exist, create it;
diff --git a/Engine/Plugins/Developer/RiderSourceCodeAccess/RiderSourceCodeAccess.uplugin b/Engine/Plugins/Developer/RiderSourceCodeAccess/RiderSourceCodeAccess.uplugin
--- a/Engine/Plugins/Developer/RiderSourceCodeAccess/RiderSourceCodeAccess.uplugin
+++ b/Engine/Plugins/Developer/RiderSourceCodeAccess/RiderSourceCodeAccess.uplugin
@@ -1,14 +1,14 @@
 {
 	"FileVersion": 3,
 	"Version": 2,
-	"VersionName": "1.5",
+	"VersionName": "1.7",
 	"FriendlyName": "Rider Integration",
 	"Description": "Allows access to source code in Rider.",
 	"Category": "Programming",
 	"CreatedBy": "JetBrains",
 	"CreatedByURL": "https://www.jetbrains.com/",
 	"DocsURL": "[REDACTED_PROFILE_URL]/RiderSourceCodeAccess",
-	"MarketplaceURL" : "",
+	"MarketplaceURL" : "com.epicgames.launcher://ue/marketplace/product/35b0d39e4cee4bccb2dabf04a4c911bb",
 	"SupportURL": "",
 	"EnabledByDefault": true,
 	"CanContainContent": false,
@@ -19,7 +19,7 @@
 			"Name": "RiderSourceCodeAccess",
 			"Type": "EditorNoCommandlet",
 			"LoadingPhase": "Default",
-			"PlatformAllowList": [ "Win64", "Mac", "Linux" ]
+			"WhitelistPlatforms": [ "Win64", "Mac", "Linux" ]
 		}
 	]
 }
diff --git a/Engine/Plugins/Developer/RiderSourceCodeAccess/Source/RiderSourceCodeAccess/Private/RiderPathLocator/Linux/RiderPathLocatorLinux.cpp b/Engine/Plugins/Developer/RiderSourceCodeAccess/Source/RiderSourceCodeAccess/Private/RiderPathLocator/Linux/RiderPathLocatorLinux.cpp
--- a/Engine/Plugins/Developer/RiderSourceCodeAccess/Source/RiderSourceCodeAccess/Private/RiderPathLocator/Linux/RiderPathLocatorLinux.cpp
+++ b/Engine/Plugins/Developer/RiderSourceCodeAccess/Source/RiderSourceCodeAccess/Private/RiderPathLocator/Linux/RiderPathLocatorLinux.cpp
@@ -4,7 +4,6 @@
 
 #include "Internationalization/Regex.h"
 #include "HAL/FileManager.h"
-#include "HAL/PlatformProcess.h"
 #include "Misc/FileHelper.h"
 #include "Misc/Paths.h"
 #include "Serialization/JsonSerializer.h"
@@ -19,7 +18,7 @@ TOptional<FInstallInfo> FRiderPathLocator::GetInstallInfoFromRiderPath(const FSt
 	{
 		return {};
 	}
-
+	
 	const FString PatternString(TEXT("(.*)(?:\\\\|/)bin"));
 	const FRegexPattern Pattern(PatternString);
 	FRegexMatcher RiderPathMatcher(Pattern, Path);
@@ -34,7 +33,7 @@ TOptional<FInstallInfo> FRiderPathLocator::GetInstallInfoFromRiderPath(const FSt
 	{
 		return {};
 	}
-
+	
 	FInstallInfo Info;
 	Info.Path = Path;
 	Info.InstallType = InstallType;
@@ -67,60 +66,49 @@ static FString GetHomePath()
 	return FHomePath;
 }
 
-// Installing the standalone version of Rider does not set up path variables, and it puts Rider in /opt
-static TArray<FInstallInfo> GetStandaloneInstalledRiders()
+static TArray<FInstallInfo> GetManuallyInstalledRiders()
 {
-	TArray<FInstallInfo> Results;
+	TArray<FInstallInfo> Result;
+	TArray<FString> RiderPaths;
 
-	const FString OptPath = TEXT("/opt");
-	const FString OptPathMask = FPaths::Combine(OptPath, TEXT("*rider*"));
+	const FString FHomePath = GetHomePath();
+	const FString HomePathMask = FPaths::Combine(FHomePath, TEXT("Rider.sh"));
 
-	TArray<FString> FolderPaths;
-	IFileManager::Get().FindFiles(FolderPaths, *OptPathMask, false, true);
+	IFileManager::Get().FindFiles(RiderPaths, *HomePathMask, false, true);
 
-	for (const FString& RiderPath : FolderPaths)
+	for(const FString& RiderPath: RiderPaths)
 	{
-		const FString FolderPath = FPaths::Combine(OptPath, RiderPath);
-
-		TArray<FString> FullPaths;
-		IFileManager::Get().FindFilesRecursive(FullPaths, *FolderPath, TEXT("rider.sh"), true, false, false);
-
-		for (const FString& FullPath : FullPaths)
+		FString FullPath = FPaths::Combine(FHomePath, RiderPath);
+		TOptional<FInstallInfo> InstallInfo = FRiderPathLocator::GetInstallInfoFromRiderPath(FullPath, FInstallInfo::EInstallType::Installed);
+		if(InstallInfo.IsSet())
 		{
-			TOptional<FInstallInfo> InstallInfo = FRiderPathLocator::GetInstallInfoFromRiderPath(FullPath, FInstallInfo::EInstallType::Installed);
-			if(InstallInfo.IsSet())
-			{
-				Results.Add(InstallInfo.GetValue());
-			}
+			Result.Add(InstallInfo.GetValue());
 		}
 	}
 
-	return Results;
-}
+	const FString FOptPath = TEXT("/opt");
+	const FString OptPathMask = FPaths::Combine(FOptPath, TEXT("Rider.sh"));
 
-// Find any Rider installations on the PATH
-static TArray<FInstallInfo> GetInstalledRidersFromPath()
-{
-	TArray<FInstallInfo> Results;
-
-	// Search each folder on the PATH
-	TArray<FString> Paths;
-	FPlatformMisc::GetEnvironmentVariable(TEXT("PATH")).ParseIntoArray(Paths, FPlatformMisc::GetPathVarDelimiter());
+	IFileManager::Get().FindFiles(RiderPaths, *OptPathMask, false, true);
 
-	for (const FString& Path : Paths)
+	for(const FString& RiderPath: RiderPaths)
 	{
-		FString RiderPath = FPaths::Combine(Path, TEXT("rider.sh"));
-		if (FPaths::FileExists(RiderPath))
+		FString FullPath = FPaths::Combine(FOptPath, RiderPath);
+		TOptional<FInstallInfo> InstallInfo = FRiderPathLocator::GetInstallInfoFromRiderPath(FullPath, FInstallInfo::EInstallType::Installed);
+		if(InstallInfo.IsSet())
 		{
-			TOptional<FInstallInfo> InstallInfo = FRiderPathLocator::GetInstallInfoFromRiderPath(RiderPath, FInstallInfo::EInstallType::Installed);
-			if(InstallInfo.IsSet())
-			{
-				Results.Add(InstallInfo.GetValue());
-			}
+			Result.Add(InstallInfo.GetValue());
 		}
 	}
 
-	return Results;
+	FString FullPath = TEXT("/snap/rider/current/bin/rider.sh");
+	TOptional<FInstallInfo> InstallInfo = FRiderPathLocator::GetInstallInfoFromRiderPath(FullPath, FInstallInfo::EInstallType::Installed);
+	if(InstallInfo.IsSet())
+	{
+		Result.Add(InstallInfo.GetValue());
+	}
+
+	return Result;
 }
 
 static FString GetToolboxPath()
@@ -174,8 +162,7 @@ TSet<FInstallInfo> FRiderPathLocator::CollectAllPaths()
 {
 	TSet<FInstallInfo> InstallInfos;
 	InstallInfos.Append(GetInstalledRidersWithMdfind());
-	InstallInfos.Append(GetStandaloneInstalledRiders());
-	InstallInfos.Append(GetInstalledRidersFromPath());
+	InstallInfos.Append(GetManuallyInstalledRiders());
 	InstallInfos.Append(GetInstallInfosFromToolbox(GetToolboxPath(), "Rider.sh"));
 	InstallInfos.Append(GetInstallInfosFromResourceFile());
 	return InstallInfos;
diff --git a/Engine/Plugins/Developer/RiderSourceCodeAccess/Source/RiderSourceCodeAccess/Private/RiderPathLocator/Mac/RiderPathLocatorMac.cpp b/Engine/Plugins/Developer/RiderSourceCodeAccess/Source/RiderSourceCodeAccess/Private/RiderPathLocator/Mac/RiderPathLocatorMac.cpp
--- a/Engine/Plugins/Developer/RiderSourceCodeAccess/Source/RiderSourceCodeAccess/Private/RiderPathLocator/Mac/RiderPathLocatorMac.cpp
+++ b/Engine/Plugins/Developer/RiderSourceCodeAccess/Source/RiderSourceCodeAccess/Private/RiderPathLocator/Mac/RiderPathLocatorMac.cpp
@@ -3,7 +3,6 @@
 #include "RiderPathLocator/RiderPathLocator.h"
 
 #include "HAL/FileManager.h"
-#include "HAL/PlatformProcess.h"
 #include "Misc/FileHelper.h"
 #include "Misc/Paths.h"
 #include "Serialization/JsonSerializer.h"
diff --git a/Engine/Plugins/Developer/RiderSourceCodeAccess/Source/RiderSourceCodeAccess/Private/RiderPathLocator/Win/RiderPathLocatorWin.cpp b/Engine/Plugins/Developer/RiderSourceCodeAccess/Source/RiderSourceCodeAccess/Private/RiderPathLocator/Win/RiderPathLocatorWin.cpp
--- a/Engine/Plugins/Developer/RiderSourceCodeAccess/Source/RiderSourceCodeAccess/Private/RiderPathLocator/Win/RiderPathLocatorWin.cpp
+++ b/Engine/Plugins/Developer/RiderSourceCodeAccess/Source/RiderSourceCodeAccess/Private/RiderPathLocator/Win/RiderPathLocatorWin.cpp
@@ -6,7 +6,7 @@
 #include "Misc/FileHelper.h"
 #include "Misc/Paths.h"
 #include "Serialization/JsonSerializer.h"
-#include "Runtime/Launch/Resources/Version.h"
+#include "Launch/Resources/Version.h"
 
 #if PLATFORM_WINDOWS
 
@@ -201,4 +201,4 @@ TSet<FInstallInfo> FRiderPathLocator::CollectAllPaths()
 	InstallInfos.Append(GetInstallInfosFromResourceFile());
 	return InstallInfos;
 }
-#endif
+#endif
\ No newline at end of file
diff --git a/Engine/Plugins/Developer/RiderSourceCodeAccess/Source/RiderSourceCodeAccess/Private/RiderSourceCodeAccessor.cpp b/Engine/Plugins/Developer/RiderSourceCodeAccess/Source/RiderSourceCodeAccess/Private/RiderSourceCodeAccessor.cpp
--- a/Engine/Plugins/Developer/RiderSourceCodeAccess/Source/RiderSourceCodeAccess/Private/RiderSourceCodeAccessor.cpp
+++ b/Engine/Plugins/Developer/RiderSourceCodeAccess/Source/RiderSourceCodeAccess/Private/RiderSourceCodeAccessor.cpp
@@ -26,7 +26,6 @@ DEFINE_LOG_CATEGORY_STATIC(LogRiderAccessor, Log, All);
 
 namespace RSCA
 {
-
 TOptional<FString> ResolvePathToFile(const FString& FullPath)
 {
 	FString Path = FullPath;
@@ -63,49 +62,42 @@ struct FCommandLineInfo
 };
 
 #if PLATFORM_MAC
-static int ProcessIsTranslated()
-{
-	int Return = 0;
-
-	size_t Size = sizeof(Return);
-	if (sysctlbyname("sysctl.proc_translated", &Return, &Size, NULL, 0) == -1)
+int processIsTranslated() {
+	int ret = 0;
+	size_t size = sizeof(ret);
+	if (sysctlbyname("sysctl.proc_translated", &ret, &size, NULL, 0) == -1) 
 	{
 		if (errno == ENOENT)
-		{
 			return 0;
-		}
 		return -1;
 	}
-
-	return Return;
+	return ret;
 }
-#endif //PLATFORM_MAC
+#endif
 
 FCommandLineInfo GetPlatformAppAndArgs(const FString& App, const FString& Args)
 {
 	FCommandLineInfo info;
 	info.App = App;
     info.Args = Args;
-
 #if PLATFORM_MAC
-	if (ProcessIsTranslated() == 1)
+	if(processIsTranslated() == 1)
 	{
 		info.App = TEXT("/usr/bin/arch");
 		info.Args = FString::Printf(TEXT("-arm64 \"%s\" %s"), *App, *Args);
 	}
-#endif // PLATFORM_MAC
-
+#endif
 	return info;
 }
 
 bool CheckExecutable(const FString& App)
 {
-	if (FPaths::FileExists(App) || FPaths::DirectoryExists(App))
+	if(FPaths::FileExists(App) || FPaths::DirectoryExists(App))
 	{
 		return true;
 	}
 
-	FNotificationInfo Info(FText::Format(LOCTEXT("CodeAccessorAppDoesntExist", "{0} doesn't exist"), FText::FromString(App)));
+	FNotificationInfo Info( FText::Format(LOCTEXT("CodeAccessorAppDoesntExist", "{0} doesn't exist"), FText::FromString(App)) );
 	Info.bFireAndForget = true;
 
 	FSlateNotificationManager::Get().AddNotification(Info)->SetCompletionState(SNotificationItem::CS_Fail);
@@ -113,23 +105,21 @@ bool CheckExecutable(const FString& App)
 	return false;
 }
 
-bool OpenRider(FString const& ExecutablePath, FString const& Params, FString const& ErrorMessage)
-{
+bool OpenRider(const FString& ExecutablePath, const FString& Params, const FString& ErrorMessage)
+{	
 	const FCommandLineInfo PlatformAppAndArgs = GetPlatformAppAndArgs(ExecutablePath, Params);
-
-	if (!CheckExecutable(PlatformAppAndArgs.App))
+	if(!CheckExecutable(PlatformAppAndArgs.App))
 	{
 		return false;
 	}
-
-	FProcHandle Proc = FPlatformProcess::CreateProc(*PlatformAppAndArgs.App, *PlatformAppAndArgs.Args, true, true, false, nullptr, 0, nullptr, nullptr);
+	FProcHandle Proc = FPlatformProcess::CreateProc(*PlatformAppAndArgs.App, *PlatformAppAndArgs.Args, true, true, false, nullptr, 0,
+													nullptr, nullptr);
 	const bool bResult = Proc.IsValid();
 	if (!bResult)
 	{
 		UE_LOG(LogRiderAccessor, Warning, TEXT("%s"), *ErrorMessage);
 		FPlatformProcess::CloseProc(Proc);
 	}
-
 	return bResult;
 }
 
@@ -213,7 +203,7 @@ bool FRiderSourceCodeAccessor::OpenSolution()
 	const FString ErrorMessage = FString::Printf(TEXT("Opening solution (%s) failed."), *FullPath);
 
 	return HandleOpeningRider([this, &Params, &ErrorMessage]()->bool
-	{
+	{		
 		return RSCA::OpenRider(ExecutablePath, Params, ErrorMessage);
 	});
 }
@@ -222,11 +212,7 @@ bool FRiderSourceCodeAccessor::OpenSolutionAtPath(const FString& InSolutionPath)
 	if (!bHasRiderInstalled) return false;
 
 	FString CorrectSolutionPath = InSolutionPath;
-	if (Model == EProjectModel::Uproject && !CorrectSolutionPath.EndsWith(".uproject"))
-	{
-		CorrectSolutionPath += ".uproject";
-	}
-	else if (!CorrectSolutionPath.EndsWith(".sln"))
+	if (!CorrectSolutionPath.EndsWith(".sln"))
 	{
 		CorrectSolutionPath += ".sln";
 	}
@@ -234,7 +220,7 @@ bool FRiderSourceCodeAccessor::OpenSolutionAtPath(const FString& InSolutionPath)
 	const FString ErrorMessage = FString::Printf(TEXT("Opening the project file (%s) failed."), *CorrectSolutionPath);
 
 	return HandleOpeningRider([this, &Params, &ErrorMessage]()->bool
-	{
+	{		
 		return RSCA::OpenRider(ExecutablePath, Params, ErrorMessage);
 	});
 }
@@ -272,7 +258,7 @@ bool FRiderSourceCodeAccessor::OpenSourceFiles(const TArray<FString>& AbsoluteSo
 	const FString ErrorMessage = FString::Printf(TEXT("Opening files (%s) failed."), *FilePaths);
 
 	return HandleOpeningRider([this, &Params, &ErrorMessage]()->bool
-	{
+	{		
 		return RSCA::OpenRider(ExecutablePath, Params, ErrorMessage);
 	});
 }
@@ -333,12 +319,12 @@ void FRiderSourceCodeAccessor::CachePathToSln() const
 
 		if (!FUProjectDictionary(FPaths::RootDir()).IsForeignProject(CachedSolutionPath))
 		{
-			FString PrimaryProjectName;
-			if (!FFileHelper::LoadFileToString(PrimaryProjectName, *(FPaths::EngineIntermediateDir() / TEXT("ProjectFiles/PrimaryProjectName.txt"))))
+			FString MasterProjectName;
+			if (!FFileHelper::LoadFileToString(MasterProjectName, *(FPaths::EngineIntermediateDir() / TEXT("ProjectFiles/MasterProjectName.txt"))))
 			{
-				PrimaryProjectName = "UE5";
+				MasterProjectName = "UE5";
 			}
-			CachedSolutionPath = FPaths::Combine(FPaths::RootDir(), PrimaryProjectName + TEXT(".sln"));
+			CachedSolutionPath = FPaths::Combine(FPaths::RootDir(), MasterProjectName + TEXT(".sln"));
 		}
 		else
 		{
diff --git a/Engine/Plugins/Developer/RiderSourceCodeAccess/Source/RiderSourceCodeAccess/Private/RiderSourceCodeAccessorModule.cpp b/Engine/Plugins/Developer/RiderSourceCodeAccess/Source/RiderSourceCodeAccess/Private/RiderSourceCodeAccessorModule.cpp
--- a/Engine/Plugins/Developer/RiderSourceCodeAccess/Source/RiderSourceCodeAccess/Private/RiderSourceCodeAccessorModule.cpp
+++ b/Engine/Plugins/Developer/RiderSourceCodeAccess/Source/RiderSourceCodeAccess/Private/RiderSourceCodeAccessorModule.cpp
@@ -7,17 +7,13 @@
 
 #include "Modules/ModuleManager.h"
 #include "Features/IModularFeatures.h"
-#include "HAL/LowLevelMemTracker.h"
 
 #define LOCTEXT_NAMESPACE "RiderSourceCodeAccessor"
 
-LLM_DEFINE_TAG(RiderSourceCodeAccess);
-
 IMPLEMENT_MODULE(FRiderSourceCodeAccessModule, RiderSourceCodeAccess);
 
 void FRiderSourceCodeAccessModule::StartupModule()
 {
-	LLM_SCOPE_BYTAG(RiderSourceCodeAccess);
 	TArray<FInstallInfo> InstallInfos = FRiderPathLocator::CollectAllPaths().Array();
 	InstallInfos.Sort();
 	GenerateUprojectAccessors(InstallInfos);
diff --git a/Engine/Plugins/Developer/RiderSourceCodeAccess/Source/RiderSourceCodeAccess/RiderSourceCodeAccess.Build.cs b/Engine/Plugins/Developer/RiderSourceCodeAccess/Source/RiderSourceCodeAccess/RiderSourceCodeAccess.Build.cs
--- a/Engine/Plugins/Developer/RiderSourceCodeAccess/Source/RiderSourceCodeAccess/RiderSourceCodeAccess.Build.cs
+++ b/Engine/Plugins/Developer/RiderSourceCodeAccess/Source/RiderSourceCodeAccess/RiderSourceCodeAccess.Build.cs
@@ -1,4 +1,4 @@
-// Copyright Epic Games, Inc. All Rights Reserved.
+﻿// Copyright Epic Games, Inc. All Rights Reserved.
 
 namespace UnrealBuildTool.Rules
 {
@@ -22,12 +22,12 @@ namespace UnrealBuildTool.Rules
 
 			if (Target.Type == TargetType.Editor)
 			{
+				#if UE_5_0_OR_LATER
 				PrivateDependencyModuleNames.Add("EditorFramework");
+				#endif
 				PrivateDependencyModuleNames.Add("UnrealEd");
 				PrivateDependencyModuleNames.Add("GameProjectGeneration");
 			}
-
-			ShortName = "RiderSCA";
 		}
 	}
 }
\ No newline at end of file
