[Backout] - CL23176028
Original CL Desc
-----------------------------------------------------------------
PR #9811: Updating RiderSourceCodeAccess (Contributed by [REDACTED_PERSON])

diff --git a/Engine/Plugins/Developer/RiderSourceCodeAccess/README.md b/Engine/Plugins/Developer/RiderSourceCodeAccess/README.md
--- a/Engine/Plugins/Developer/RiderSourceCodeAccess/README.md
+++ b/Engine/Plugins/Developer/RiderSourceCodeAccess/README.md
@@ -16,7 +16,7 @@ RiderSourceCodeAccess is a plugin for Unreal Engine, available through [Marketpl
 
 # How to add this plugin manually (not recommened)
 ## How to add to Game project
-1. Get the source code of the [master branch]([REDACTED_PROFILE_URL]/RiderSourceCodeAccess/archive/refs/heads/master.zip); 
+1. Go to ["Releases" page]([REDACTED_PROFILE_URL]/RiderSourceCodeAccess/releases) and download version of plugin for your version of Unreal Engine; 
 2. Unzip `RiderSourceCodeAccess.zip` to `RiderSourceCodeAccess` folder;
 3. Copy `RiderSourceCodeAccess` folder to `{GameProjectRoot}/Plugins/Developer`;
   a. If `{GameProjectRoot}/Plugins/Developer` does not exist, create it;
@@ -28,7 +28,7 @@ RiderSourceCodeAccess is a plugin for Unreal Engine, available through [Marketpl
 
 ## How to add to Engine project
 ### Unreal Engine is built from source code   
-1. Get the source code of the [master branch]([REDACTED_PROFILE_URL]/RiderSourceCodeAccess/archive/refs/heads/master.zip); 
+1. Go to ["Releases" page]([REDACTED_PROFILE_URL]/RiderSourceCodeAccess/releases) and download version of plugin for your version of Unreal Engine; 
 2. Unzip `RiderSourceCodeAccess.zip` to `RiderSourceCodeAccess` folder;
 3. Copy `RiderSourceCodeAccess` folder to `{UnrealEngineRoot}/Engine/Plugins/Developer`;
   a. If `{UnrealEngineRoot}/Engine/Plugins/Developer` does not exist, create it;
@@ -39,7 +39,7 @@ RiderSourceCodeAccess is a plugin for Unreal Engine, available through [Marketpl
 8. Select Rider from drop down list. NB: only Rider with C++ plugins will be available from drop down list.
 
 ### Unreal Engine is downloaded from Epic Games Store
-1. Get the source code of the [master branch]([REDACTED_PROFILE_URL]/RiderSourceCodeAccess/archive/refs/heads/master.zip); 
+1. Go to ["Releases" page]([REDACTED_PROFILE_URL]/RiderSourceCodeAccess/releases) and download version of plugin for your version of Unreal Engine; 
 2. Unzip `RiderSourceCodeAccess.zip` to `RiderSourceCodeAccess` folder;
 3. Copy `RiderSourceCodeAccess` folder to `{UnrealProjectRoot}/Plugins/Developer`;
   a. If `{UnrealProjectRoot}/Plugins/Developer` does not exist, create it;
diff --git a/Engine/Plugins/Developer/RiderSourceCodeAccess/RiderSourceCodeAccess.uplugin b/Engine/Plugins/Developer/RiderSourceCodeAccess/RiderSourceCodeAccess.uplugin
--- a/Engine/Plugins/Developer/RiderSourceCodeAccess/RiderSourceCodeAccess.uplugin
+++ b/Engine/Plugins/Developer/RiderSourceCodeAccess/RiderSourceCodeAccess.uplugin
@@ -1,14 +1,14 @@
 {
 	"FileVersion": 3,
 	"Version": 2,
-	"VersionName": "1.7",
+	"VersionName": "1.5",
 	"FriendlyName": "Rider Integration",
 	"Description": "Allows access to source code in Rider.",
 	"Category": "Programming",
 	"CreatedBy": "JetBrains",
 	"CreatedByURL": "https://www.jetbrains.com/",
 	"DocsURL": "[REDACTED_PROFILE_URL]/RiderSourceCodeAccess",
-	"MarketplaceURL" : "com.epicgames.launcher://ue/marketplace/product/35b0d39e4cee4bccb2dabf04a4c911bb",
+	"MarketplaceURL" : "",
 	"SupportURL": "",
 	"EnabledByDefault": true,
 	"CanContainContent": false,
@@ -19,7 +19,7 @@
 			"Name": "RiderSourceCodeAccess",
 			"Type": "EditorNoCommandlet",
 			"LoadingPhase": "Default",
-			"WhitelistPlatforms": [ "Win64", "Mac", "Linux" ]
+			"PlatformAllowList": [ "Win64", "Mac", "Linux" ]
 		}
 	]
 }
diff --git a/Engine/Plugins/Developer/RiderSourceCodeAccess/Source/RiderSourceCodeAccess/Private/RiderPathLocator/Linux/RiderPathLocatorLinux.cpp b/Engine/Plugins/Developer/RiderSourceCodeAccess/Source/RiderSourceCodeAccess/Private/RiderPathLocator/Linux/RiderPathLocatorLinux.cpp
--- a/Engine/Plugins/Developer/RiderSourceCodeAccess/Source/RiderSourceCodeAccess/Private/RiderPathLocator/Linux/RiderPathLocatorLinux.cpp
+++ b/Engine/Plugins/Developer/RiderSourceCodeAccess/Source/RiderSourceCodeAccess/Private/RiderPathLocator/Linux/RiderPathLocatorLinux.cpp
@@ -4,6 +4,7 @@
 
 #include "Internationalization/Regex.h"
 #include "HAL/FileManager.h"
+#include "HAL/PlatformProcess.h"
 #include "Misc/FileHelper.h"
 #include "Misc/Paths.h"
 #include "Serialization/JsonSerializer.h"
@@ -18,7 +19,7 @@ TOptional<FInstallInfo> FRiderPathLocator::GetInstallInfoFromRiderPath(const FSt
 	{
 		return {};
 	}
-	
+
 	const FString PatternString(TEXT("(.*)(?:\\\\|/)bin"));
 	const FRegexPattern Pattern(PatternString);
 	FRegexMatcher RiderPathMatcher(Pattern, Path);
@@ -33,7 +34,7 @@ TOptional<FInstallInfo> FRiderPathLocator::GetInstallInfoFromRiderPath(const FSt
 	{
 		return {};
 	}
-	
+
 	FInstallInfo Info;
 	Info.Path = Path;
 	Info.InstallType = InstallType;
@@ -66,49 +67,60 @@ static FString GetHomePath()
 	return FHomePath;
 }
 
-static TArray<FInstallInfo> GetManuallyInstalledRiders()
+// Installing the standalone version of Rider does not set up path variables, and it puts Rider in /opt
+static TArray<FInstallInfo> GetStandaloneInstalledRiders()
 {
-	TArray<FInstallInfo> Result;
-	TArray<FString> RiderPaths;
+	TArray<FInstallInfo> Results;
 
-	const FString FHomePath = GetHomePath();
-	const FString HomePathMask = FPaths::Combine(FHomePath, TEXT("Rider.sh"));
+	const FString OptPath = TEXT("/opt");
+	const FString OptPathMask = FPaths::Combine(OptPath, TEXT("*rider*"));
 
-	IFileManager::Get().FindFiles(RiderPaths, *HomePathMask, false, true);
+	TArray<FString> FolderPaths;
+	IFileManager::Get().FindFiles(FolderPaths, *OptPathMask, false, true);
 
-	for(const FString& RiderPath: RiderPaths)
+	for (const FString& RiderPath : FolderPaths)
 	{
-		FString FullPath = FPaths::Combine(FHomePath, RiderPath);
-		TOptional<FInstallInfo> InstallInfo = FRiderPathLocator::GetInstallInfoFromRiderPath(FullPath, FInstallInfo::EInstallType::Installed);
-		if(InstallInfo.IsSet())
+		const FString FolderPath = FPaths::Combine(OptPath, RiderPath);
+
+		TArray<FString> FullPaths;
+		IFileManager::Get().FindFilesRecursive(FullPaths, *FolderPath, TEXT("rider.sh"), true, false, false);
+
+		for (const FString& FullPath : FullPaths)
 		{
-			Result.Add(InstallInfo.GetValue());
+			TOptional<FInstallInfo> InstallInfo = FRiderPathLocator::GetInstallInfoFromRiderPath(FullPath, FInstallInfo::EInstallType::Installed);
+			if(InstallInfo.IsSet())
+			{
+				Results.Add(InstallInfo.GetValue());
+			}
 		}
 	}
 
-	const FString FOptPath = TEXT("/opt");
-	const FString OptPathMask = FPaths::Combine(FOptPath, TEXT("Rider.sh"));
+	return Results;
+}
 
-	IFileManager::Get().FindFiles(RiderPaths, *OptPathMask, false, true);
+// Find any Rider installations on the PATH
+static TArray<FInstallInfo> GetInstalledRidersFromPath()
+{
+	TArray<FInstallInfo> Results;
 
-	for(const FString& RiderPath: RiderPaths)
+	// Search each folder on the PATH
+	TArray<FString> Paths;
+	FPlatformMisc::GetEnvironmentVariable(TEXT("PATH")).ParseIntoArray(Paths, FPlatformMisc::GetPathVarDelimiter());
+
+	for (const FString& Path : Paths)
 	{
-		FString FullPath = FPaths::Combine(FOptPath, RiderPath);
-		TOptional<FInstallInfo> InstallInfo = FRiderPathLocator::GetInstallInfoFromRiderPath(FullPath, FInstallInfo::EInstallType::Installed);
-		if(InstallInfo.IsSet())
+		FString RiderPath = FPaths::Combine(Path, TEXT("rider.sh"));
+		if (FPaths::FileExists(RiderPath))
 		{
-			Result.Add(InstallInfo.GetValue());
+			TOptional<FInstallInfo> InstallInfo = FRiderPathLocator::GetInstallInfoFromRiderPath(RiderPath, FInstallInfo::EInstallType::Installed);
+			if(InstallInfo.IsSet())
+			{
+				Results.Add(InstallInfo.GetValue());
+			}
 		}
 	}
 
-	FString FullPath = TEXT("/snap/rider/current/bin/rider.sh");
-	TOptional<FInstallInfo> InstallInfo = FRiderPathLocator::GetInstallInfoFromRiderPath(FullPath, FInstallInfo::EInstallType::Installed);
-	if(InstallInfo.IsSet())
-	{
-		Result.Add(InstallInfo.GetValue());
-	}
-
-	return Result;
+	return Results;
 }
 
 static FString GetToolboxPath()
@@ -162,7 +174,8 @@ TSet<FInstallInfo> FRiderPathLocator::CollectAllPaths()
 {
 	TSet<FInstallInfo> InstallInfos;
 	InstallInfos.Append(GetInstalledRidersWithMdfind());
-	InstallInfos.Append(GetManuallyInstalledRiders());
+	InstallInfos.Append(GetStandaloneInstalledRiders());
+	InstallInfos.Append(GetInstalledRidersFromPath());
 	InstallInfos.Append(GetInstallInfosFromToolbox(GetToolboxPath(), "Rider.sh"));
 	InstallInfos.Append(GetInstallInfosFromResourceFile());
 	return InstallInfos;
diff --git a/Engine/Plugins/Developer/RiderSourceCodeAccess/Source/RiderSourceCodeAccess/Private/RiderPathLocator/Mac/RiderPathLocatorMac.cpp b/Engine/Plugins/Developer/RiderSourceCodeAccess/Source/RiderSourceCodeAccess/Private/RiderPathLocator/Mac/RiderPathLocatorMac.cpp
--- a/Engine/Plugins/Developer/RiderSourceCodeAccess/Source/RiderSourceCodeAccess/Private/RiderPathLocator/Mac/RiderPathLocatorMac.cpp
+++ b/Engine/Plugins/Developer/RiderSourceCodeAccess/Source/RiderSourceCodeAccess/Private/RiderPathLocator/Mac/RiderPathLocatorMac.cpp
@@ -3,6 +3,7 @@
 #include "RiderPathLocator/RiderPathLocator.h"
 
 #include "HAL/FileManager.h"
+#include "HAL/PlatformProcess.h"
 #include "Misc/FileHelper.h"
 #include "Misc/Paths.h"
 #include "Serialization/JsonSerializer.h"
diff --git a/Engine/Plugins/Developer/RiderSourceCodeAccess/Source/RiderSourceCodeAccess/Private/RiderPathLocator/Win/RiderPathLocatorWin.cpp b/Engine/Plugins/Developer/RiderSourceCodeAccess/Source/RiderSourceCodeAccess/Private/RiderPathLocator/Win/RiderPathLocatorWin.cpp
--- a/Engine/Plugins/Developer/RiderSourceCodeAccess/Source/RiderSourceCodeAccess/Private/RiderPathLocator/Win/RiderPathLocatorWin.cpp
+++ b/Engine/Plugins/Developer/RiderSourceCodeAccess/Source/RiderSourceCodeAccess/Private/RiderPathLocator/Win/RiderPathLocatorWin.cpp
@@ -6,7 +6,7 @@
 #include "Misc/FileHelper.h"
 #include "Misc/Paths.h"
 #include "Serialization/JsonSerializer.h"
-#include "Launch/Resources/Version.h"
+#include "Runtime/Launch/Resources/Version.h"
 
 #if PLATFORM_WINDOWS
 
@@ -201,4 +201,4 @@ TSet<FInstallInfo> FRiderPathLocator::CollectAllPaths()
 	InstallInfos.Append(GetInstallInfosFromResourceFile());
 	return InstallInfos;
 }
-#endif
\ No newline at end of file
+#endif
diff --git a/Engine/Plugins/Developer/RiderSourceCodeAccess/Source/RiderSourceCodeAccess/Private/RiderSourceCodeAccessor.cpp b/Engine/Plugins/Developer/RiderSourceCodeAccess/Source/RiderSourceCodeAccess/Private/RiderSourceCodeAccessor.cpp
--- a/Engine/Plugins/Developer/RiderSourceCodeAccess/Source/RiderSourceCodeAccess/Private/RiderSourceCodeAccessor.cpp
+++ b/Engine/Plugins/Developer/RiderSourceCodeAccess/Source/RiderSourceCodeAccess/Private/RiderSourceCodeAccessor.cpp
@@ -26,6 +26,7 @@ DEFINE_LOG_CATEGORY_STATIC(LogRiderAccessor, Log, All);
 
 namespace RSCA
 {
+
 TOptional<FString> ResolvePathToFile(const FString& FullPath)
 {
 	FString Path = FullPath;
@@ -62,42 +63,49 @@ struct FCommandLineInfo
 };
 
 #if PLATFORM_MAC
-int processIsTranslated() {
-	int ret = 0;
-	size_t size = sizeof(ret);
-	if (sysctlbyname("sysctl.proc_translated", &ret, &size, NULL, 0) == -1) 
+static int ProcessIsTranslated()
+{
+	int Return = 0;
+
+	size_t Size = sizeof(Return);
+	if (sysctlbyname("sysctl.proc_translated", &Return, &Size, NULL, 0) == -1)
 	{
 		if (errno == ENOENT)
+		{
 			return 0;
+		}
 		return -1;
 	}
-	return ret;
+
+	return Return;
 }
-#endif
+#endif //PLATFORM_MAC
 
 FCommandLineInfo GetPlatformAppAndArgs(const FString& App, const FString& Args)
 {
 	FCommandLineInfo info;
 	info.App = App;
     info.Args = Args;
+
 #if PLATFORM_MAC
-	if(processIsTranslated() == 1)
+	if (ProcessIsTranslated() == 1)
 	{
 		info.App = TEXT("/usr/bin/arch");
 		info.Args = FString::Printf(TEXT("-arm64 \"%s\" %s"), *App, *Args);
 	}
-#endif
+#endif // PLATFORM_MAC
+
 	return info;
 }
 
 bool CheckExecutable(const FString& App)
 {
-	if(FPaths::FileExists(App) || FPaths::DirectoryExists(App))
+	if (FPaths::FileExists(App) || FPaths::DirectoryExists(App))
 	{
 		return true;
 	}
 
-	FNotificationInfo Info( FText::Format(LOCTEXT("CodeAccessorAppDoesntExist", "{0} doesn't exist"), FText::FromString(App)) );
+	FNotificationInfo Info(FText::Format(LOCTEXT("CodeAccessorAppDoesntExist", "{0} doesn't exist"), FText::FromString(App)));
 	Info.bFireAndForget = true;
 
 	FSlateNotificationManager::Get().AddNotification(Info)->SetCompletionState(SNotificationItem::CS_Fail);
@@ -105,21 +113,23 @@ bool CheckExecutable(const FString& App)
 	return false;
 }
 
-bool OpenRider(const FString& ExecutablePath, const FString& Params, const FString& ErrorMessage)
-{	
+bool OpenRider(FString const& ExecutablePath, FString const& Params, FString const& ErrorMessage)
+{
 	const FCommandLineInfo PlatformAppAndArgs = GetPlatformAppAndArgs(ExecutablePath, Params);
-	if(!CheckExecutable(PlatformAppAndArgs.App))
+
+	if (!CheckExecutable(PlatformAppAndArgs.App))
 	{
 		return false;
 	}
-	FProcHandle Proc = FPlatformProcess::CreateProc(*PlatformAppAndArgs.App, *PlatformAppAndArgs.Args, true, true, false, nullptr, 0,
-													nullptr, nullptr);
+
+	FProcHandle Proc = FPlatformProcess::CreateProc(*PlatformAppAndArgs.App, *PlatformAppAndArgs.Args, true, true, false, nullptr, 0, nullptr, nullptr);
 	const bool bResult = Proc.IsValid();
 	if (!bResult)
 	{
 		UE_LOG(LogRiderAccessor, Warning, TEXT("%s"), *ErrorMessage);
 		FPlatformProcess::CloseProc(Proc);
 	}
+
 	return bResult;
 }
 
@@ -203,7 +213,7 @@ bool FRiderSourceCodeAccessor::OpenSolution()
 	const FString ErrorMessage = FString::Printf(TEXT("Opening solution (%s) failed."), *FullPath);
 
 	return HandleOpeningRider([this, &Params, &ErrorMessage]()->bool
-	{		
+	{
 		return RSCA::OpenRider(ExecutablePath, Params, ErrorMessage);
 	});
 }
@@ -212,7 +222,11 @@ bool FRiderSourceCodeAccessor::OpenSolutionAtPath(const FString& InSolutionPath)
 	if (!bHasRiderInstalled) return false;
 
 	FString CorrectSolutionPath = InSolutionPath;
-	if (!CorrectSolutionPath.EndsWith(".sln"))
+	if (Model == EProjectModel::Uproject && !CorrectSolutionPath.EndsWith(".uproject"))
+	{
+		CorrectSolutionPath += ".uproject";
+	}
+	else if (!CorrectSolutionPath.EndsWith(".sln"))
 	{
 		CorrectSolutionPath += ".sln";
 	}
@@ -220,7 +234,7 @@ bool FRiderSourceCodeAccessor::OpenSolutionAtPath(const FString& InSolutionPath)
 	const FString ErrorMessage = FString::Printf(TEXT("Opening the project file (%s) failed."), *CorrectSolutionPath);
 
 	return HandleOpeningRider([this, &Params, &ErrorMessage]()->bool
-	{		
+	{
 		return RSCA::OpenRider(ExecutablePath, Params, ErrorMessage);
 	});
 }
@@ -258,7 +272,7 @@ bool FRiderSourceCodeAccessor::OpenSourceFiles(const TArray<FString>& AbsoluteSo
 	const FString ErrorMessage = FString::Printf(TEXT("Opening files (%s) failed."), *FilePaths);
 
 	return HandleOpeningRider([this, &Params, &ErrorMessage]()->bool
-	{		
+	{
 		return RSCA::OpenRider(ExecutablePath, Params, ErrorMessage);
 	});
 }
@@ -319,12 +333,12 @@ void FRiderSourceCodeAccessor::CachePathToSln() const
 
 		if (!FUProjectDictionary(FPaths::RootDir()).IsForeignProject(CachedSolutionPath))
 		{
-			FString MasterProjectName;
-			if (!FFileHelper::LoadFileToString(MasterProjectName, *(FPaths::EngineIntermediateDir() / TEXT("ProjectFiles/MasterProjectName.txt"))))
+			FString PrimaryProjectName;
+			if (!FFileHelper::LoadFileToString(PrimaryProjectName, *(FPaths::EngineIntermediateDir() / TEXT("ProjectFiles/PrimaryProjectName.txt"))))
 			{
-				MasterProjectName = "UE5";
+				PrimaryProjectName = "UE5";
 			}
-			CachedSolutionPath = FPaths::Combine(FPaths::RootDir(), MasterProjectName + TEXT(".sln"));
+			CachedSolutionPath = FPaths::Combine(FPaths::RootDir(), PrimaryProjectName + TEXT(".sln"));
 		}
 		else
 		{
diff --git a/Engine/Plugins/Developer/RiderSourceCodeAccess/Source/RiderSourceCodeAccess/Private/RiderSourceCodeAccessorModule.cpp b/Engine/Plugins/Developer/RiderSourceCodeAccess/Source/RiderSourceCodeAccess/Private/RiderSourceCodeAccessorModule.cpp
--- a/Engine/Plugins/Developer/RiderSourceCodeAccess/Source/RiderSourceCodeAccess/Private/RiderSourceCodeAccessorModule.cpp
+++ b/Engine/Plugins/Developer/RiderSourceCodeAccess/Source/RiderSourceCodeAccess/Private/RiderSourceCodeAccessorModule.cpp
@@ -7,13 +7,17 @@
 
 #include "Modules/ModuleManager.h"
 #include "Features/IModularFeatures.h"
+#include "HAL/LowLevelMemTracker.h"
 
 #define LOCTEXT_NAMESPACE "RiderSourceCodeAccessor"
 
+LLM_DEFINE_TAG(RiderSourceCodeAccess);
+
 IMPLEMENT_MODULE(FRiderSourceCodeAccessModule, RiderSourceCodeAccess);
 
 void FRiderSourceCodeAccessModule::StartupModule()
 {
+	LLM_SCOPE_BYTAG(RiderSourceCodeAccess);
 	TArray<FInstallInfo> InstallInfos = FRiderPathLocator::CollectAllPaths().Array();
 	InstallInfos.Sort();
 	GenerateUprojectAccessors(InstallInfos);
diff --git a/Engine/Plugins/Developer/RiderSourceCodeAccess/Source/RiderSourceCodeAccess/RiderSourceCodeAccess.Build.cs b/Engine/Plugins/Developer/RiderSourceCodeAccess/Source/RiderSourceCodeAccess/RiderSourceCodeAccess.Build.cs
--- a/Engine/Plugins/Developer/RiderSourceCodeAccess/Source/RiderSourceCodeAccess/RiderSourceCodeAccess.Build.cs
+++ b/Engine/Plugins/Developer/RiderSourceCodeAccess/Source/RiderSourceCodeAccess/RiderSourceCodeAccess.Build.cs
@@ -1,4 +1,4 @@
-﻿// Copyright Epic Games, Inc. All Rights Reserved.
+// Copyright Epic Games, Inc. All Rights Reserved.
 
 namespace UnrealBuildTool.Rules
 {
@@ -22,12 +22,12 @@ namespace UnrealBuildTool.Rules
 
 			if (Target.Type == TargetType.Editor)
 			{
-				#if UE_5_0_OR_LATER
 				PrivateDependencyModuleNames.Add("EditorFramework");
-				#endif
 				PrivateDependencyModuleNames.Add("UnrealEd");
 				PrivateDependencyModuleNames.Add("GameProjectGeneration");
 			}
+
+			ShortName = "RiderSCA";
 		}
 	}
 }
\ No newline at end of file
