PR #10511: Fix rsca for toolbox 2 0

diff --git a/Engine/Plugins/Developer/RiderSourceCodeAccess/Source/RiderSourceCodeAccess/Private/RiderPathLocator/Common/RiderPathLocator.cpp b/Engine/Plugins/Developer/RiderSourceCodeAccess/Source/RiderSourceCodeAccess/Private/RiderPathLocator/Common/RiderPathLocator.cpp
--- a/Engine/Plugins/Developer/RiderSourceCodeAccess/Source/RiderSourceCodeAccess/Private/RiderPathLocator/Common/RiderPathLocator.cpp
+++ b/Engine/Plugins/Developer/RiderSourceCodeAccess/Source/RiderSourceCodeAccess/Private/RiderPathLocator/Common/RiderPathLocator.cpp
@@ -1,13 +1,16 @@
 ﻿// Copyright Epic Games, Inc. All Rights Reserved.
 
 #include "RiderPathLocator/RiderPathLocator.h"
+#include "Dom/JsonObject.h"
 #include "Interfaces/IPluginManager.h"
+#include "Internationalization/Regex.h"
 #include "Misc/FileHelper.h"
+#include "Misc/Paths.h"
+#include "Serialization/JsonReader.h"
 #include "Serialization/JsonSerializer.h"
 
-TArray<FInstallInfo> FRiderPathLocator::GetInstallInfosFromToolbox(const FString& ToolboxPath, const FString& Pattern)
+FString ExtractPathFromSettingsJson(const FString& ToolboxPath)
 {
-	FString InstallPath = ToolboxPath;
 	const FString SettingJsonPath = FPaths::Combine(ToolboxPath, TEXT(".settings.json"));
 	if (FPaths::FileExists(SettingJsonPath))
 	{
@@ -22,30 +25,95 @@ TArray<FInstallInfo> FRiderPathLocator::GetInstallInfosFromToolbox(const FString
 			{
 				if (!InstallLocation.IsEmpty())
 				{
-					InstallPath = InstallLocation;
+					return InstallLocation;
 				}
 			}
 		}
 	}
+	return {};
+}
+
+TArray<FInstallInfo> FRiderPathLocator::GetInstallInfosFromToolbox(const FString& ToolboxPath, const FString& Pattern)
+{
+	if(!FPaths::DirectoryExists(ToolboxPath)) return {};
 	
-	const FString ToolboxRiderRootPath = FPaths::Combine(InstallPath, TEXT("apps"));
-	if(!FPaths::DirectoryExists(ToolboxRiderRootPath)) return {};
+	const FString InstallLocationPath = ExtractPathFromSettingsJson(ToolboxPath);
+	TArray<FInstallInfo> Result{};
+	if(!InstallLocationPath.IsEmpty())
+	{
+		// Toolbox V1 custom install location search path
+		Result = GetInstallInfos(FPaths::Combine(InstallLocationPath, TEXT("apps")), Pattern, FInstallInfo::EInstallType::Toolbox);
+		if(Result.Num() != 0) return Result;
+
+		// Toolbox V2 custom install location search path
+		return GetInstallInfos(InstallLocationPath, Pattern, FInstallInfo::EInstallType::Toolbox);		
+	}
+
+	// Toolbox V1 default install location search path
+	Result = GetInstallInfos(FPaths::Combine(ToolboxPath, TEXT("apps")), Pattern, FInstallInfo::EInstallType::Toolbox);
+	if(Result.Num() != 0) return Result;
 
-	return GetInstallInfos(ToolboxRiderRootPath, Pattern, FInstallInfo::EInstallType::Toolbox);
+	const FString DefaultInstallLocation = GetDefaultIDEInstallLocationForToolboxV2();
+	return GetInstallInfos(DefaultInstallLocation, Pattern, FInstallInfo::EInstallType::Toolbox);
+}
+
+FVersion FRiderPathLocator::GetLastBuildVersion(const FString& HistoryJsonPath)
+{
+	if(!FPaths::FileExists(HistoryJsonPath)) return {};
+	
+	FString JsonStr;
+	FFileHelper::LoadFileToString(JsonStr, *HistoryJsonPath);
+	const TSharedRef<TJsonReader<TCHAR>> JsonReader = TJsonReaderFactory<TCHAR>::Create(JsonStr);
+	TSharedPtr<FJsonObject> JsonObject = MakeShareable(new FJsonObject());
+	if (!FJsonSerializer::Deserialize(JsonReader, JsonObject) || !JsonObject.IsValid()) return {};
+	
+	const TArray< TSharedPtr<FJsonValue> >* HistoryField;
+	if(!JsonObject->TryGetArrayField(TEXT("history"), HistoryField)) return {};
+	if(HistoryField->Num() == 0) return {};
+
+	const TSharedPtr<FJsonObject>* LastItemObject;
+	if(!HistoryField->Last()->TryGetObject(LastItemObject)) return {};
+
+	const TSharedPtr<FJsonObject>* ItemObject;
+	if(!LastItemObject->Get()->TryGetObjectField("item", ItemObject)) return {};
+
+	FString Build;
+	if(!ItemObject->Get()->TryGetStringField("build", Build)) return {};
+
+	return FVersion(Build);
+}
+
+FString FRiderPathLocator::GetHistoryJsonPath(const FString& RiderPath)
+{
+	FString Directory = FPaths::ConvertRelativePathToFull(FPaths::Combine(RiderPath, ".."));
+	int8_t SafeCheck = 10;
+	while(FPaths::DirectoryExists(Directory) && SafeCheck-- > 0)
+	{
+		FString HistoryPath = FPaths::Combine(Directory, ".history.json");
+		if(FPaths::FileExists(HistoryPath)) return HistoryPath;
+
+		Directory = FPaths::ConvertRelativePathToFull(FPaths::Combine(Directory, ".."));
+	}
+	return {};
 }
 
 TArray<FInstallInfo> FRiderPathLocator::GetInstallInfos(const FString& ToolboxRiderRootPath, const FString& Pattern, FInstallInfo::EInstallType InstallType)
 {
+	if(ToolboxRiderRootPath.IsEmpty() || !FPaths::DirectoryExists(ToolboxRiderRootPath)) return {};
+	
 	TArray<FInstallInfo> RiderInstallInfos;
 	TArray<FString> RiderPaths;
 	IFileManager::Get().FindFilesRecursive(RiderPaths, *ToolboxRiderRootPath, *Pattern, true, true);
 	for(const FString& RiderPath: RiderPaths)
 	{
 		TOptional<FInstallInfo> InstallInfo = GetInstallInfoFromRiderPath(RiderPath, InstallType);
-		if(InstallInfo.IsSet())
-		{
-			RiderInstallInfos.Add(InstallInfo.GetValue());
-		}
+		if(!InstallInfo.IsSet()) continue;
+		
+		FString HistoryJsonPath = GetHistoryJsonPath(RiderPath);
+		FVersion Version = GetLastBuildVersion(HistoryJsonPath);
+		if(Version.IsInitialized() && InstallInfo->Version != Version) continue;
+		
+		RiderInstallInfos.Add(InstallInfo.GetValue());
 	}
 	return RiderInstallInfos;
 }
diff --git a/Engine/Plugins/Developer/RiderSourceCodeAccess/Source/RiderSourceCodeAccess/Private/RiderPathLocator/Linux/RiderPathLocatorLinux.cpp b/Engine/Plugins/Developer/RiderSourceCodeAccess/Source/RiderSourceCodeAccess/Private/RiderPathLocator/Linux/RiderPathLocatorLinux.cpp
--- a/Engine/Plugins/Developer/RiderSourceCodeAccess/Source/RiderSourceCodeAccess/Private/RiderPathLocator/Linux/RiderPathLocatorLinux.cpp
+++ b/Engine/Plugins/Developer/RiderSourceCodeAccess/Source/RiderSourceCodeAccess/Private/RiderPathLocator/Linux/RiderPathLocatorLinux.cpp
@@ -8,10 +8,17 @@
 
 #include "Internationalization/Regex.h"
 #include "HAL/FileManager.h"
-#include "Linux/LinuxPlatformProcess.h"
+#include "HAL/PlatformProcess.h"
+#include "Misc/Paths.h"
 
 #include "Runtime/Launch/Resources/Version.h"
 
+FString FRiderPathLocator::GetDefaultIDEInstallLocationForToolboxV2()
+{
+	// V2 and V1 have the same path on Linux, we don't need to process it extra
+	return {};
+}
+
 TOptional<FInstallInfo> FRiderPathLocator::GetInstallInfoFromRiderPath(const FString& Path, FInstallInfo::EInstallType InstallType)
 {
 	if(!FPaths::FileExists(Path))
@@ -71,34 +78,27 @@ static TArray<FInstallInfo> GetManuallyInstalledRiders()
 	TArray<FInstallInfo> Result;
 	TArray<FString> RiderPaths;
 
-	const FString FHomePath = GetHomePath();
-	const FString HomePathMask = FPaths::Combine(FHomePath, TEXT("Rider.sh"));
+	const TArray RiderLookupPaths = {
+		GetHomePath(),
+		FString(TEXT("/opt")),
+		FPaths::Combine(TEXT("/usr"), TEXT("local"), TEXT("bin"))
+	};
 
-	IFileManager::Get().FindFiles(RiderPaths, *HomePathMask, false, true);
-
-	for(const FString& RiderPath: RiderPaths)
+	for(const FString& RiderLookupPath: RiderLookupPaths)
 	{
-		FString FullPath = FPaths::Combine(FHomePath, RiderPath);
-		TOptional<FInstallInfo> InstallInfo = FRiderPathLocator::GetInstallInfoFromRiderPath(FullPath, FInstallInfo::EInstallType::Installed);
-		if(InstallInfo.IsSet())
-		{
-			Result.Add(InstallInfo.GetValue());
-		}
-	}
-
-	const FString FOptPath = TEXT("/opt");
-	const FString OptPathMask = FPaths::Combine(FOptPath, TEXT("Rider.sh"));
-
-	IFileManager::Get().FindFiles(RiderPaths, *OptPathMask, false, true);
+		FString RiderLookupPathMask = FPaths::Combine(RiderLookupPath,TEXT("*Rider*"));
+		IFileManager::Get().FindFiles(RiderPaths, *RiderLookupPathMask, false, true);
 
-	for(const FString& RiderPath: RiderPaths)
-	{
-		FString FullPath = FPaths::Combine(FOptPath, RiderPath);
-		TOptional<FInstallInfo> InstallInfo = FRiderPathLocator::GetInstallInfoFromRiderPath(FullPath, FInstallInfo::EInstallType::Installed);
-		if(InstallInfo.IsSet())
+		for(const FString& RiderPath: RiderPaths)
 		{
-			Result.Add(InstallInfo.GetValue());
+			FString FullPath = FPaths::Combine(RiderLookupPath, RiderPath, TEXT("bin"), TEXT("rider.sh"));
+			TOptional<FInstallInfo> InstallInfo = FRiderPathLocator::GetInstallInfoFromRiderPath(FullPath, FInstallInfo::EInstallType::Installed);
+			if(InstallInfo.IsSet())
+			{
+				Result.Add(InstallInfo.GetValue());
+			}
 		}
+		RiderPaths.Empty();
 	}
 
 	FString FullPath = TEXT("/snap/rider/current/bin/rider.sh");
@@ -119,19 +119,17 @@ static FString GetToolboxPath()
 	return FPaths::Combine(LocalAppData, TEXT("JetBrains"), TEXT("Toolbox"));
 }
 
-static TArray<FInstallInfo> GetInstalledRidersWithMdfind()
+static TArray<FInstallInfo> GetInstalledRidersWithLocate()
 {
 	int32 ReturnCode;
 	FString OutResults;
 	FString OutErrors;
-
-	// avoid trying to run mdfind if it doesnt exists
-	if (!FPaths::FileExists(TEXT("/usr/bin/mdfind")))
+	if (!FPaths::FileExists(TEXT("/usr/bin/locate")))
 	{
 		return {};
 	}
 
-	FPlatformProcess::ExecProcess(TEXT("/usr/bin/mdfind"), TEXT("\"kMDItemKind == Application\""), &ReturnCode, &OutResults, &OutErrors);
+	FPlatformProcess::ExecProcess(TEXT("/usr/bin/locate"), TEXT("-e bin/rider.sh"), &ReturnCode, &OutResults, &OutErrors);
 	if (ReturnCode != 0)
 	{
 		return {};
@@ -141,11 +139,13 @@ static TArray<FInstallInfo> GetInstalledRidersWithMdfind()
 	FString TmpString;
 	while(OutResults.Split(TEXT("\n"), &TmpString, &OutResults))
 	{
-		if(TmpString.Contains(TEXT("Rider")))
+		if(TmpString.Contains(TEXT("snapd")) || TmpString.Contains(TEXT(".local")) || TmpString.Contains(TEXT("/opt")))
 		{
-			RiderPaths.Add(TmpString);
+			continue;
 		}
+		RiderPaths.Add(TmpString);
 	}
+
 	TArray<FInstallInfo> Result;
 	for(const FString& RiderPath: RiderPaths)
 	{
@@ -161,7 +161,7 @@ static TArray<FInstallInfo> GetInstalledRidersWithMdfind()
 TSet<FInstallInfo> FRiderPathLocator::CollectAllPaths()
 {
 	TSet<FInstallInfo> InstallInfos;
-	InstallInfos.Append(GetInstalledRidersWithMdfind());
+	InstallInfos.Append(GetInstalledRidersWithLocate());
 	InstallInfos.Append(GetManuallyInstalledRiders());
 	InstallInfos.Append(GetInstallInfosFromToolbox(GetToolboxPath(), "Rider.sh"));
 	InstallInfos.Append(GetInstallInfosFromResourceFile());
diff --git a/Engine/Plugins/Developer/RiderSourceCodeAccess/Source/RiderSourceCodeAccess/Private/RiderPathLocator/Mac/RiderPathLocatorMac.cpp b/Engine/Plugins/Developer/RiderSourceCodeAccess/Source/RiderSourceCodeAccess/Private/RiderPathLocator/Mac/RiderPathLocatorMac.cpp
--- a/Engine/Plugins/Developer/RiderSourceCodeAccess/Source/RiderSourceCodeAccess/Private/RiderPathLocator/Mac/RiderPathLocatorMac.cpp
+++ b/Engine/Plugins/Developer/RiderSourceCodeAccess/Source/RiderSourceCodeAccess/Private/RiderPathLocator/Mac/RiderPathLocatorMac.cpp
@@ -7,10 +7,8 @@
 #include "RiderPathLocator/RiderPathLocator.h"
 
 #include "HAL/FileManager.h"
-#include "Misc/FileHelper.h"
+#include "HAL/PlatformProcess.h"
 #include "Misc/Paths.h"
-#include "Serialization/JsonSerializer.h"
-#include "Mac/MacPlatformProcess.h"
 
 #include "Runtime/Launch/Resources/Version.h"
 
@@ -109,6 +107,18 @@ static TArray<FInstallInfo> GetInstalledRidersWithMdfind()
 	return Result;
 }
 
+FString FRiderPathLocator::GetDefaultIDEInstallLocationForToolboxV2()
+{
+#if ENGINE_MAJOR_VERSION == 4 && ENGINE_MINOR_VERSION <= 20
+	TCHAR CHomePath[4096];
+	FPlatformMisc::GetEnvironmentVariable(TEXT("HOME"), CHomePath, ARRAY_COUNT(CHomePath));
+	const FString FHomePath = CHomePath;
+#else
+	const FString FHomePath = FPlatformMisc::GetEnvironmentVariable(TEXT("HOME"));
+#endif
+	return FPaths::Combine(FHomePath, TEXT("Applications"));
+}
+
 TSet<FInstallInfo> FRiderPathLocator::CollectAllPaths()
 {
 	TSet<FInstallInfo> InstallInfos;
diff --git a/Engine/Plugins/Developer/RiderSourceCodeAccess/Source/RiderSourceCodeAccess/Private/RiderPathLocator/RiderPathLocator.h b/Engine/Plugins/Developer/RiderSourceCodeAccess/Source/RiderSourceCodeAccess/Private/RiderPathLocator/RiderPathLocator.h
--- a/Engine/Plugins/Developer/RiderSourceCodeAccess/Source/RiderSourceCodeAccess/Private/RiderPathLocator/RiderPathLocator.h
+++ b/Engine/Plugins/Developer/RiderSourceCodeAccess/Source/RiderSourceCodeAccess/Private/RiderPathLocator/RiderPathLocator.h
@@ -2,15 +2,20 @@
 
 #pragma once
 
-#include "Containers/Set.h"
+#include "CoreMinimal.h"
 #include "Misc/Paths.h"
 
-template <typename OptionalType> struct TOptional;
-
 struct FVersion
 {	
 	bool IsInitialized() const { return Versions.Num() != 0; }
 
+	FVersion(){}
+
+	explicit FVersion(const FString& VersionString)
+	{
+		*this = VersionString;
+	}
+
 	void operator=(const FString& VersionString)
 	{
 		TArray<FString> Output;
@@ -61,7 +66,16 @@ struct FVersion
 			if(LeftElement > RightElement) return false;
 		}
 		return LeftSize < RightSize;
-		
+	}
+
+	bool operator==(const FVersion& rhs) const
+	{
+		return !(*this < rhs) && !(rhs < *this); 
+	}
+
+	bool operator!=(const FVersion& rhs) const
+	{
+		return (*this < rhs) || (rhs < *this); 
 	}
 
 	FString ToString() const
@@ -122,6 +136,7 @@ struct FInstallInfo
     }
 };
 
+template <typename OptionalType> struct TOptional;
 class FRiderPathLocator
 {
 public:
@@ -130,7 +145,10 @@ public:
 	static TSet<FInstallInfo> CollectAllPaths();
 private:
 	static void ParseProductInfoJson(FInstallInfo& Info, const FString& ProductInfoJsonPath);
+	static FString GetDefaultIDEInstallLocationForToolboxV2();
 	static TArray<FInstallInfo> GetInstallInfosFromToolbox(const FString& ToolboxPath, const FString& Pattern);
 	static TArray<FInstallInfo> GetInstallInfosFromResourceFile();
 	static TArray<FInstallInfo> GetInstallInfos(const FString& ToolboxRiderRootPath, const FString& Pattern, FInstallInfo::EInstallType InstallType);
+	static FString GetHistoryJsonPath(const FString& RiderPath);
+	static FVersion GetLastBuildVersion(const FString& HistoryJsonPath);
 };
diff --git a/Engine/Plugins/Developer/RiderSourceCodeAccess/Source/RiderSourceCodeAccess/Private/RiderPathLocator/Win/RiderPathLocatorWin.cpp b/Engine/Plugins/Developer/RiderSourceCodeAccess/Source/RiderSourceCodeAccess/Private/RiderPathLocator/Win/RiderPathLocatorWin.cpp
--- a/Engine/Plugins/Developer/RiderSourceCodeAccess/Source/RiderSourceCodeAccess/Private/RiderPathLocator/Win/RiderPathLocatorWin.cpp
+++ b/Engine/Plugins/Developer/RiderSourceCodeAccess/Source/RiderSourceCodeAccess/Private/RiderPathLocator/Win/RiderPathLocatorWin.cpp
@@ -3,11 +3,11 @@
 #include "HAL/Platform.h"
 
 #if PLATFORM_WINDOWS
-
 #include "RiderPathLocator/RiderPathLocator.h"
+
 #include "Internationalization/Regex.h"
-#include "Misc/FileHelper.h"
-#include "Serialization/JsonSerializer.h"
+#include "Misc/Paths.h"
+
 #include "Runtime/Launch/Resources/Version.h"
 #include "Windows/WindowsPlatformMisc.h"
 #include "Windows/AllowWindowsPlatformTypes.h"
@@ -45,6 +45,12 @@ static FString GetToolboxPath(const Windows::HKEY RootKey, const FString& Regist
 	return ToolboxPathMatcher.GetCaptureGroup(1);
 }
 
+FString FRiderPathLocator::GetDefaultIDEInstallLocationForToolboxV2()
+{
+	const FString LocalAppData = FPlatformMisc::GetEnvironmentVariable(TEXT("LOCALAPPDATA"));
+	return FPaths::Combine(LocalAppData, TEXT("Programs"));
+}
+
 static bool EnumerateRegistryKeys(HKEY Key, TArray<FString> &OutNames)
 {
 	for (DWORD Index = 0;; Index++)
@@ -104,41 +110,76 @@ static LONG GetStringRegKey(const HKEY Key, const FString& ValueName, FString& V
 	}
 	return Result;
 }
-
 static TArray<FInstallInfo> CollectPathsFromRegistry( const Windows::HKEY RootKey, const FString& RegistryKey)
 {
 	TArray<FInstallInfo> InstallInfos;
 	HKEY Key;
 	const LONG Result = RegOpenKeyEx(RootKey, *RegistryKey, 0, KEY_READ, &Key);
-	if (Result == ERROR_SUCCESS)
+	if (Result != ERROR_SUCCESS) return InstallInfos;
+	
+	TArray<FString> Keys;
+	if (!EnumerateRegistryKeys(Key, Keys)) return InstallInfos;
+	
+	for (const FString& key : Keys)
 	{
-		TArray<FString> Keys;
-		if (EnumerateRegistryKeys(Key, Keys))
+		if (!key.Contains(TEXT("Rider"))) continue;
+
+		HKEY SubKey;
+		const LONG SubResult = RegOpenKeyEx(Key, *key, 0, KEY_READ, &SubKey);
+		if (SubResult != ERROR_SUCCESS) continue;
+
+		TArray<FString> Values;
+		if (!EnumerateRegistryValues(SubKey, Values)) continue;
+
+		for (const auto& Value : Values)
 		{
-			for (const FString& key : Keys)
+			if (Value != TEXT("InstallLocation")) continue;
+			FString InstallLocation;
+			if (GetStringRegKey(SubKey, Value, InstallLocation) != ERROR_SUCCESS) continue;
+			
+			const FString ExePath = FPaths::Combine(InstallLocation, TEXT("bin"), TEXT("rider64.exe"));
+			TOptional<FInstallInfo> InstallInfo = FRiderPathLocator::GetInstallInfoFromRiderPath(ExePath, FInstallInfo::EInstallType::Installed);
+			if(InstallInfo.IsSet())
 			{
-				if (!key.Contains(TEXT("Rider"))) continue;
-
-				HKEY SubKey;
-				const LONG SubResult = RegOpenKeyEx(Key, *key, 0, KEY_READ, &SubKey);
-				if (SubResult != ERROR_SUCCESS) continue;
-
-				TArray<FString> Values;
-				if (!EnumerateRegistryValues(SubKey, Values)) continue;
-
-				for (const auto& Value : Values)
-				{
-					if (Value != TEXT("InstallLocation")) continue;
-					FString InstallLocation;
-					if (GetStringRegKey(SubKey, Value, InstallLocation) != ERROR_SUCCESS) continue;
-					const FString ExePath = FPaths::Combine(InstallLocation, TEXT("bin"), TEXT("rider64.exe"));
-					TOptional<FInstallInfo> InstallInfo = FRiderPathLocator::GetInstallInfoFromRiderPath(ExePath, FInstallInfo::EInstallType::Installed);
-					if(InstallInfo.IsSet())
-					{
-						InstallInfos.Add(InstallInfo.GetValue());
-					}
-				}
+				InstallInfos.Add(InstallInfo.GetValue());
+			}
+		}
+	}
+
+	return InstallInfos;
+}
 
+static TArray<FInstallInfo> CollectDotUltimatePathsFromRegistry( const Windows::HKEY RootKey, const FString& RegistryKey)
+{
+	TArray<FInstallInfo> InstallInfos;
+	HKEY Key;
+	const LONG Result = RegOpenKeyEx(RootKey, *RegistryKey, 0, KEY_READ, &Key);
+	if (Result != ERROR_SUCCESS) return InstallInfos;
+	
+	TArray<FString> Keys;
+	if (!EnumerateRegistryKeys(Key, Keys)) return InstallInfos;
+	
+	for (const FString& key : Keys)
+	{
+		HKEY SubKey;
+		const LONG SubResult = RegOpenKeyEx(Key, *key, 0, KEY_READ, &SubKey);
+		if (SubResult != ERROR_SUCCESS) continue;
+
+		TArray<FString> Values;
+		if (!EnumerateRegistryValues(SubKey, Values)) continue;
+
+		for (const auto& Value : Values)
+		{
+			if (Value != TEXT("InstallDir")) continue;
+			
+			FString InstallLocation;
+			if (GetStringRegKey(SubKey, Value, InstallLocation) != ERROR_SUCCESS) continue;
+			
+			const FString ExePath = FPaths::Combine(InstallLocation, TEXT("bin"), TEXT("rider64.exe"));
+			TOptional<FInstallInfo> InstallInfo = FRiderPathLocator::GetInstallInfoFromRiderPath(ExePath, FInstallInfo::EInstallType::Installed);
+			if(InstallInfo.IsSet())
+			{
+				InstallInfos.Add(InstallInfo.GetValue());
 			}
 		}
 	}
@@ -190,15 +231,19 @@ TOptional<FInstallInfo> FRiderPathLocator::GetInstallInfoFromRiderPath(const FSt
 TSet<FInstallInfo> FRiderPathLocator::CollectAllPaths()
 {
 	TSet<FInstallInfo> InstallInfos;
+	InstallInfos.Append(CollectPathsFromRegistry(HKEY_CURRENT_USER, TEXT("SOFTWARE\\Microsoft\\Windows\\CurrentVersion\\Uninstall")));
+	InstallInfos.Append(CollectPathsFromRegistry(HKEY_LOCAL_MACHINE, TEXT("SOFTWARE\\Microsoft\\Windows\\CurrentVersion\\Uninstall")));
+	InstallInfos.Append(CollectPathsFromRegistry(HKEY_CURRENT_USER, TEXT("SOFTWARE\\WOW6432Node\\Microsoft\\Windows\\CurrentVersion\\Uninstall")));
+	InstallInfos.Append(CollectPathsFromRegistry(HKEY_LOCAL_MACHINE, TEXT("SOFTWARE\\WOW6432Node\\Microsoft\\Windows\\CurrentVersion\\Uninstall")));
+	InstallInfos.Append(CollectDotUltimatePathsFromRegistry(HKEY_CURRENT_USER, TEXT("SOFTWARE\\JetBrains\\Rider")));
+	InstallInfos.Append(CollectDotUltimatePathsFromRegistry(HKEY_LOCAL_MACHINE, TEXT("SOFTWARE\\JetBrains\\Rider")));
+	InstallInfos.Append(CollectDotUltimatePathsFromRegistry(HKEY_CURRENT_USER, TEXT("SOFTWARE\\WOW6432Node\\JetBrains\\Rider")));
+	InstallInfos.Append(CollectDotUltimatePathsFromRegistry(HKEY_LOCAL_MACHINE, TEXT("SOFTWARE\\WOW6432Node\\JetBrains\\Rider")));
 	InstallInfos.Append(GetInstallInfosFromToolbox(GetToolboxPath(), "rider64.exe"));
 	InstallInfos.Append(GetInstallInfosFromToolbox(GetToolboxPath(HKEY_CURRENT_USER, TEXT("Software\\JetBrains\\Toolbox\\")), "rider64.exe"));
 	InstallInfos.Append(GetInstallInfosFromToolbox(GetToolboxPath(HKEY_LOCAL_MACHINE, TEXT("Software\\JetBrains\\Toolbox\\")), "rider64.exe"));
 	InstallInfos.Append(GetInstallInfosFromToolbox(GetToolboxPath(HKEY_CURRENT_USER, TEXT("Software\\JetBrains s.r.o.\\JetBrainsToolbox\\")), "rider64.exe"));
 	InstallInfos.Append(GetInstallInfosFromToolbox(GetToolboxPath(HKEY_LOCAL_MACHINE, TEXT("Software\\JetBrains s.r.o.\\JetBrainsToolbox\\")), "rider64.exe"));
-	InstallInfos.Append(CollectPathsFromRegistry(HKEY_CURRENT_USER, TEXT("SOFTWARE\\Microsoft\\Windows\\CurrentVersion\\Uninstall")));
-	InstallInfos.Append(CollectPathsFromRegistry(HKEY_LOCAL_MACHINE, TEXT("SOFTWARE\\Microsoft\\Windows\\CurrentVersion\\Uninstall")));
-	InstallInfos.Append(CollectPathsFromRegistry(HKEY_CURRENT_USER, TEXT("SOFTWARE\\WOW6432Node\\Microsoft\\Windows\\CurrentVersion\\Uninstall")));
-	InstallInfos.Append(CollectPathsFromRegistry(HKEY_LOCAL_MACHINE, TEXT("SOFTWARE\\WOW6432Node\\Microsoft\\Windows\\CurrentVersion\\Uninstall")));
 	InstallInfos.Append(GetInstallInfosFromResourceFile());
 	return InstallInfos;
 }
diff --git a/Engine/Plugins/Developer/RiderSourceCodeAccess/Source/RiderSourceCodeAccess/Private/RiderSourceCodeAccessor.cpp b/Engine/Plugins/Developer/RiderSourceCodeAccess/Source/RiderSourceCodeAccess/Private/RiderSourceCodeAccessor.cpp
--- a/Engine/Plugins/Developer/RiderSourceCodeAccess/Source/RiderSourceCodeAccess/Private/RiderSourceCodeAccessor.cpp
+++ b/Engine/Plugins/Developer/RiderSourceCodeAccess/Source/RiderSourceCodeAccess/Private/RiderSourceCodeAccessor.cpp
@@ -4,9 +4,13 @@
 
 #include "RiderPathLocator/RiderPathLocator.h"
 
+#include "Modules/ModuleManager.h"
 #include "Misc/App.h"
 #include "Misc/FileHelper.h"
+#include "Misc/Paths.h"
+#include "Misc/ScopeLock.h"
 #include "Misc/UProjectInfo.h"
+#include "DesktopPlatformModule.h"
 #include "ISourceCodeAccessModule.h"
 #include "Interfaces/IProjectManager.h"
 #include "ProjectDescriptor.h"
diff --git a/Engine/Plugins/Developer/RiderSourceCodeAccess/Source/RiderSourceCodeAccess/Private/RiderSourceCodeAccessorModule.h b/Engine/Plugins/Developer/RiderSourceCodeAccess/Source/RiderSourceCodeAccess/Private/RiderSourceCodeAccessorModule.h
--- a/Engine/Plugins/Developer/RiderSourceCodeAccess/Source/RiderSourceCodeAccess/Private/RiderSourceCodeAccessorModule.h
+++ b/Engine/Plugins/Developer/RiderSourceCodeAccess/Source/RiderSourceCodeAccess/Private/RiderSourceCodeAccessorModule.h
@@ -2,10 +2,7 @@
 
 #pragma once
 
-#include "Modules/ModuleInterface.h"
-#include "Templates/SharedPointer.h"
-
-class ISourceCodeAccessor;
+#include "ISourceCodeAccessModule.h"
 
 class FRiderSourceCodeAccessModule : public IModuleInterface
 {
diff --git a/Engine/Plugins/Developer/RiderSourceCodeAccess/Source/RiderSourceCodeAccess/RiderSourceCodeAccess.Build.cs b/Engine/Plugins/Developer/RiderSourceCodeAccess/Source/RiderSourceCodeAccess/RiderSourceCodeAccess.Build.cs
--- a/Engine/Plugins/Developer/RiderSourceCodeAccess/Source/RiderSourceCodeAccess/RiderSourceCodeAccess.Build.cs
+++ b/Engine/Plugins/Developer/RiderSourceCodeAccess/Source/RiderSourceCodeAccess/RiderSourceCodeAccess.Build.cs
@@ -16,7 +16,8 @@ namespace UnrealBuildTool.Rules
 					"Projects",
 					"Json",
 					"Projects",
-					"Slate"
+					"Slate",
+					"SlateCore"
 				}
 			);
 
