Remove marketplace url causing build issues

#rnx

diff --git a/Engine/Plugins/Developer/RiderSourceCodeAccess/RiderSourceCodeAccess.uplugin b/Engine/Plugins/Developer/RiderSourceCodeAccess/RiderSourceCodeAccess.uplugin
--- a/Engine/Plugins/Developer/RiderSourceCodeAccess/RiderSourceCodeAccess.uplugin
+++ b/Engine/Plugins/Developer/RiderSourceCodeAccess/RiderSourceCodeAccess.uplugin
@@ -8,7 +8,7 @@
 	"CreatedBy": "JetBrains",
 	"CreatedByURL": "https://www.jetbrains.com/",
 	"DocsURL": "[REDACTED_PROFILE_URL]/RiderSourceCodeAccess",
-	"MarketplaceURL" : "com.epicgames.launcher://ue/marketplace/product/35b0d39e4cee4bccb2dabf04a4c911bb",
+	"MarketplaceURL" : "",
 	"SupportURL": "",
 	"EnabledByDefault": true,
 	"CanContainContent": false,
