RSCA plugin changes from JetBrains to avoid scanning large directories when registry keys etc are empty

diff --git a/Engine/Plugins/Developer/RiderSourceCodeAccess/Source/RiderSourceCodeAccess/Private/RiderPathLocator/Common/RiderPathLocator.cpp b/Engine/Plugins/Developer/RiderSourceCodeAccess/Source/RiderSourceCodeAccess/Private/RiderPathLocator/Common/RiderPathLocator.cpp
--- a/Engine/Plugins/Developer/RiderSourceCodeAccess/Source/RiderSourceCodeAccess/Private/RiderPathLocator/Common/RiderPathLocator.cpp
+++ b/Engine/Plugins/Developer/RiderSourceCodeAccess/Source/RiderSourceCodeAccess/Private/RiderPathLocator/Common/RiderPathLocator.cpp
@@ -33,9 +33,14 @@ FString ExtractPathFromSettingsJson(const FString& ToolboxPath)
 	return {};
 }
 
+bool FRiderPathLocator::DirectoryExistsAndNonEmpty(const FString& Path)
+{
+	return !Path.IsEmpty() && FPaths::DirectoryExists(Path);
+}
+
 TArray<FInstallInfo> FRiderPathLocator::GetInstallInfosFromToolbox(const FString& ToolboxPath, const FString& Pattern)
 {
-	if(!FPaths::DirectoryExists(ToolboxPath)) return {};
+	if(!DirectoryExistsAndNonEmpty(ToolboxPath)) return {};
 	
 	const FString InstallLocationPath = ExtractPathFromSettingsJson(ToolboxPath);
 	TArray<FInstallInfo> Result{};
@@ -87,7 +92,7 @@ FString FRiderPathLocator::GetHistoryJsonPath(const FString& RiderPath)
 {
 	FString Directory = FPaths::ConvertRelativePathToFull(FPaths::Combine(RiderPath, ".."));
 	int8_t SafeCheck = 10;
-	while(FPaths::DirectoryExists(Directory) && SafeCheck-- > 0)
+	while(DirectoryExistsAndNonEmpty(Directory) && SafeCheck-- > 0)
 	{
 		FString HistoryPath = FPaths::Combine(Directory, ".history.json");
 		if(FPaths::FileExists(HistoryPath)) return HistoryPath;
@@ -99,7 +104,7 @@ FString FRiderPathLocator::GetHistoryJsonPath(const FString& RiderPath)
 
 TArray<FInstallInfo> FRiderPathLocator::GetInstallInfos(const FString& ToolboxRiderRootPath, const FString& Pattern, FInstallInfo::EInstallType InstallType)
 {
-	if(ToolboxRiderRootPath.IsEmpty() || !FPaths::DirectoryExists(ToolboxRiderRootPath)) return {};
+	if(!DirectoryExistsAndNonEmpty(ToolboxRiderRootPath)) return {};
 	
 	TArray<FInstallInfo> RiderInstallInfos;
 	TArray<FString> RiderPaths;
diff --git a/Engine/Plugins/Developer/RiderSourceCodeAccess/Source/RiderSourceCodeAccess/Private/RiderPathLocator/Linux/RiderPathLocatorLinux.cpp b/Engine/Plugins/Developer/RiderSourceCodeAccess/Source/RiderSourceCodeAccess/Private/RiderPathLocator/Linux/RiderPathLocatorLinux.cpp
--- a/Engine/Plugins/Developer/RiderSourceCodeAccess/Source/RiderSourceCodeAccess/Private/RiderPathLocator/Linux/RiderPathLocatorLinux.cpp
+++ b/Engine/Plugins/Developer/RiderSourceCodeAccess/Source/RiderSourceCodeAccess/Private/RiderPathLocator/Linux/RiderPathLocatorLinux.cpp
@@ -36,7 +36,7 @@ TOptional<FInstallInfo> FRiderPathLocator::GetInstallInfoFromRiderPath(const FSt
 
 	const FString RiderDir = RiderPathMatcher.GetCaptureGroup(1);
 	const FString RiderCppPluginPath = FPaths::Combine(RiderDir, TEXT("plugins"), TEXT("rider-cpp"));
-	if (!FPaths::DirectoryExists(RiderCppPluginPath))
+	if (!DirectoryExistsAndNonEmpty(RiderCppPluginPath))
 	{
 		return {};
 	}
diff --git a/Engine/Plugins/Developer/RiderSourceCodeAccess/Source/RiderSourceCodeAccess/Private/RiderPathLocator/Mac/RiderPathLocatorMac.cpp b/Engine/Plugins/Developer/RiderSourceCodeAccess/Source/RiderSourceCodeAccess/Private/RiderPathLocator/Mac/RiderPathLocatorMac.cpp
--- a/Engine/Plugins/Developer/RiderSourceCodeAccess/Source/RiderSourceCodeAccess/Private/RiderPathLocator/Mac/RiderPathLocatorMac.cpp
+++ b/Engine/Plugins/Developer/RiderSourceCodeAccess/Source/RiderSourceCodeAccess/Private/RiderPathLocator/Mac/RiderPathLocatorMac.cpp
@@ -14,13 +14,13 @@
 
 TOptional<FInstallInfo> FRiderPathLocator::GetInstallInfoFromRiderPath(const FString& PathToRiderApp, FInstallInfo::EInstallType InstallType)
 {
-	if(!FPaths::DirectoryExists(PathToRiderApp))
+	if(!DirectoryExistsAndNonEmpty(PathToRiderApp))
 	{
 		return {};
 	}
 
 	const FString RiderCppPluginPath = FPaths::Combine(PathToRiderApp, TEXT("Contents"), TEXT("plugins"), TEXT("rider-cpp"));
-	if (!FPaths::DirectoryExists(RiderCppPluginPath))
+	if (!DirectoryExistsAndNonEmpty(RiderCppPluginPath))
 	{
 		return {};
 	}
diff --git a/Engine/Plugins/Developer/RiderSourceCodeAccess/Source/RiderSourceCodeAccess/Private/RiderPathLocator/RiderPathLocator.h b/Engine/Plugins/Developer/RiderSourceCodeAccess/Source/RiderSourceCodeAccess/Private/RiderPathLocator/RiderPathLocator.h
--- a/Engine/Plugins/Developer/RiderSourceCodeAccess/Source/RiderSourceCodeAccess/Private/RiderPathLocator/RiderPathLocator.h
+++ b/Engine/Plugins/Developer/RiderSourceCodeAccess/Source/RiderSourceCodeAccess/Private/RiderPathLocator/RiderPathLocator.h
@@ -142,6 +142,7 @@ class FRiderPathLocator
 public:
 	// Platform specific implementation
 	static TOptional<FInstallInfo> GetInstallInfoFromRiderPath(const FString& Path, FInstallInfo::EInstallType InstallType);
+	static bool DirectoryExistsAndNonEmpty(const FString& Path);
 	static TSet<FInstallInfo> CollectAllPaths();
 private:
 	static void ParseProductInfoJson(FInstallInfo& Info, const FString& ProductInfoJsonPath);
diff --git a/Engine/Plugins/Developer/RiderSourceCodeAccess/Source/RiderSourceCodeAccess/Private/RiderPathLocator/Win/RiderPathLocatorWin.cpp b/Engine/Plugins/Developer/RiderSourceCodeAccess/Source/RiderSourceCodeAccess/Private/RiderPathLocator/Win/RiderPathLocatorWin.cpp
--- a/Engine/Plugins/Developer/RiderSourceCodeAccess/Source/RiderSourceCodeAccess/Private/RiderPathLocator/Win/RiderPathLocatorWin.cpp
+++ b/Engine/Plugins/Developer/RiderSourceCodeAccess/Source/RiderSourceCodeAccess/Private/RiderPathLocator/Win/RiderPathLocatorWin.cpp
@@ -204,7 +204,7 @@ TOptional<FInstallInfo> FRiderPathLocator::GetInstallInfoFromRiderPath(const FSt
 
 	const FString RiderDir = RiderPathMatcher.GetCaptureGroup(1);
 	const FString RiderCppPluginPath = FPaths::Combine(RiderDir, TEXT("plugins"), TEXT("rider-cpp"));
-	if (!FPaths::DirectoryExists(RiderCppPluginPath))
+	if (!DirectoryExistsAndNonEmpty(RiderCppPluginPath))
 	{
 		return {};
 	}
