PR #11109: Specify styles when selecting icons for SourceCodeAccess implementations

diff --git a/Engine/Plugins/Developer/N10XSourceCodeAccess/Source/N10XSourceCodeAccess/N10XSourceCodeAccess.Build.cs b/Engine/Plugins/Developer/N10XSourceCodeAccess/Source/N10XSourceCodeAccess/N10XSourceCodeAccess.Build.cs
--- a/Engine/Plugins/Developer/N10XSourceCodeAccess/Source/N10XSourceCodeAccess/N10XSourceCodeAccess.Build.cs
+++ b/Engine/Plugins/Developer/N10XSourceCodeAccess/Source/N10XSourceCodeAccess/N10XSourceCodeAccess.Build.cs
@@ -14,7 +14,7 @@ namespace UnrealBuildTool.Rules
 				{
 					"Core",
 					"SourceCodeAccess",
-					"DesktopPlatform",
+					"DesktopPlatform"
 				}
 			);
 
diff --git a/Engine/Plugins/Developer/RiderSourceCodeAccess/Source/RiderSourceCodeAccess/Private/RiderSourceCodeAccessor.cpp b/Engine/Plugins/Developer/RiderSourceCodeAccess/Source/RiderSourceCodeAccess/Private/RiderSourceCodeAccessor.cpp
--- a/Engine/Plugins/Developer/RiderSourceCodeAccess/Source/RiderSourceCodeAccess/Private/RiderSourceCodeAccessor.cpp
+++ b/Engine/Plugins/Developer/RiderSourceCodeAccess/Source/RiderSourceCodeAccess/Private/RiderSourceCodeAccessor.cpp
@@ -18,6 +18,7 @@
 #include "Dialogs/SOutputLogDialog.h"
 #include "Framework/Notifications/NotificationManager.h"
 #include "Misc/MessageDialog.h"
+#include "Style/RiderSourceCodeAccessStyleSet.h"
 #include "Widgets/Notifications/SNotificationList.h"
 
 #define LOCTEXT_NAMESPACE "RiderSourceCodeAccessor"
@@ -268,6 +269,21 @@ bool FRiderSourceCodeAccessor::SaveAllOpenDocuments() const
 	return false;
 }
 
+FName FRiderSourceCodeAccessor::GetOpenIconName() const
+{
+	return FRiderSourceCodeAccessStyleSet::RiderIconName;
+}
+
+FName FRiderSourceCodeAccessor::GetRefreshIconName() const
+{
+	return FRiderSourceCodeAccessStyleSet::RiderRefreshIconName;
+}
+
+FName FRiderSourceCodeAccessor::GetStyleSet() const
+{
+	return FRiderSourceCodeAccessStyleSet::StyleName;
+}
+
 void FRiderSourceCodeAccessor::Init(const FInstallInfo& Info, EProjectModel ProjectModel, EAccessType Type)
 {
 	Model = ProjectModel; 
diff --git a/Engine/Plugins/Developer/RiderSourceCodeAccess/Source/RiderSourceCodeAccess/Private/RiderSourceCodeAccessor.h b/Engine/Plugins/Developer/RiderSourceCodeAccess/Source/RiderSourceCodeAccess/Private/RiderSourceCodeAccessor.h
--- a/Engine/Plugins/Developer/RiderSourceCodeAccess/Source/RiderSourceCodeAccess/Private/RiderSourceCodeAccessor.h
+++ b/Engine/Plugins/Developer/RiderSourceCodeAccess/Source/RiderSourceCodeAccess/Private/RiderSourceCodeAccessor.h
@@ -42,6 +42,9 @@ public:
 	virtual bool AddSourceFiles(const TArray<FString>& AbsoluteSourcePaths, const TArray<FString>& AvailableModules) override;
 	virtual bool SaveAllOpenDocuments() const override;
 	virtual void Tick(const float) override {}
+	virtual FName GetOpenIconName() const override;
+	virtual FName GetRefreshIconName() const override;
+	virtual FName GetStyleSet() const override;
 private:
 	void CachePathToUproject() const;
 	void CachePathToSln() const;
diff --git a/Engine/Plugins/Developer/RiderSourceCodeAccess/Source/RiderSourceCodeAccess/Private/RiderSourceCodeAccessorModule.cpp b/Engine/Plugins/Developer/RiderSourceCodeAccess/Source/RiderSourceCodeAccess/Private/RiderSourceCodeAccessorModule.cpp
--- a/Engine/Plugins/Developer/RiderSourceCodeAccess/Source/RiderSourceCodeAccess/Private/RiderSourceCodeAccessorModule.cpp
+++ b/Engine/Plugins/Developer/RiderSourceCodeAccess/Source/RiderSourceCodeAccess/Private/RiderSourceCodeAccessorModule.cpp
@@ -7,6 +7,7 @@
 
 #include "Modules/ModuleManager.h"
 #include "Features/IModularFeatures.h"
+#include "Style/RiderSourceCodeAccessStyleSet.h"
 
 #define LOCTEXT_NAMESPACE "RiderSourceCodeAccessor"
 
@@ -14,6 +15,8 @@ IMPLEMENT_MODULE(FRiderSourceCodeAccessModule, RiderSourceCodeAccess);
 
 void FRiderSourceCodeAccessModule::StartupModule()
 {
+	FRiderSourceCodeAccessStyleSet::Initialize();
+	
 	TArray<FInstallInfo> InstallInfos = FRiderPathLocator::CollectAllPaths().Array();
 	InstallInfos.Sort();
 	GenerateUprojectAccessors(InstallInfos);
@@ -32,6 +35,8 @@ void FRiderSourceCodeAccessModule::ShutdownModule()
 		// Unbind provider from editor
 		IModularFeatures::Get().UnregisterModularFeature(FRiderSourceCodeAccessor::FeatureType(), &(RiderSourceCodeAccessor.Value.Get()));
 	}
+	
+	FRiderSourceCodeAccessStyleSet::Shutdown();
 }
 
 void FRiderSourceCodeAccessModule::GenerateSlnAccessors(const TArray<FInstallInfo>& InstallInfos)
diff --git a/Engine/Plugins/Developer/RiderSourceCodeAccess/Source/RiderSourceCodeAccess/Private/Style/RiderSourceCodeAccessStyleSet.cpp b/Engine/Plugins/Developer/RiderSourceCodeAccess/Source/RiderSourceCodeAccess/Private/Style/RiderSourceCodeAccessStyleSet.cpp
new file mode 100644
--- /dev/null
+++ b/Engine/Plugins/Developer/RiderSourceCodeAccess/Source/RiderSourceCodeAccess/Private/Style/RiderSourceCodeAccessStyleSet.cpp
@@ -0,0 +1,53 @@
+﻿#include "RiderSourceCodeAccessStyleSet.h"
+
+#include "Interfaces/IPluginManager.h"
+#include "Styling/SlateStyleRegistry.h"
+#include "Styling/AppStyle.h"
+#include "Styling/CoreStyle.h"
+
+FName FRiderSourceCodeAccessStyleSet::RiderRefreshIconName("RSCA.RiderRefreshIcon");
+FName FRiderSourceCodeAccessStyleSet::RiderIconName("RSCA.RiderIcon");
+FName FRiderSourceCodeAccessStyleSet::StyleName("RiderSourceCodeAccessStyle");
+TUniquePtr<FRiderSourceCodeAccessStyleSet> FRiderSourceCodeAccessStyleSet::Inst(nullptr);
+
+
+const FRiderSourceCodeAccessStyleSet& FRiderSourceCodeAccessStyleSet::Get()
+{
+	ensure(Inst.IsValid());
+	return *(Inst.Get());
+}
+
+void FRiderSourceCodeAccessStyleSet::Initialize()
+{
+	if (!Inst.IsValid())
+	{
+		Inst = TUniquePtr<FRiderSourceCodeAccessStyleSet>(new FRiderSourceCodeAccessStyleSet);
+	}
+}
+
+void FRiderSourceCodeAccessStyleSet::Shutdown()
+{
+	if (Inst.IsValid())
+	{
+		FSlateStyleRegistry::UnRegisterSlateStyle(*Inst.Get());
+		Inst.Reset();
+	}
+}
+
+FRiderSourceCodeAccessStyleSet::FRiderSourceCodeAccessStyleSet() : FSlateStyleSet(StyleName)
+{
+	SetParentStyleName(FAppStyle::GetAppStyleSetName());
+	FString Path = IPluginManager::Get().FindPlugin(TEXT("RiderSourceCodeAccess"))->GetBaseDir() / TEXT("Resources");
+	FSlateStyleSet::SetContentRoot(IPluginManager::Get().FindPlugin(TEXT("RiderSourceCodeAccess"))->GetBaseDir() / TEXT("Resources"));
+
+	Set(RiderRefreshIconName, new FSlateVectorImageBrush(FSlateStyleSet::RootToContentDir(TEXT("RiderRefresh.svg")), CoreStyleConstants::Icon16x16));
+	Set(RiderIconName, new FSlateVectorImageBrush(FSlateStyleSet::RootToContentDir(TEXT("Rider.svg")), CoreStyleConstants::Icon16x16));
+	
+	FSlateStyleRegistry::RegisterSlateStyle(*this);	
+}
+
+const FName& FRiderSourceCodeAccessStyleSet::GetStyleSetName() const
+{
+	return StyleName;
+}
+
diff --git a/Engine/Plugins/Developer/RiderSourceCodeAccess/Source/RiderSourceCodeAccess/Private/Style/RiderSourceCodeAccessStyleSet.h b/Engine/Plugins/Developer/RiderSourceCodeAccess/Source/RiderSourceCodeAccess/Private/Style/RiderSourceCodeAccessStyleSet.h
new file mode 100644
--- /dev/null
+++ b/Engine/Plugins/Developer/RiderSourceCodeAccess/Source/RiderSourceCodeAccess/Private/Style/RiderSourceCodeAccessStyleSet.h
@@ -0,0 +1,26 @@
+﻿#pragma once
+
+#include "Styling/SlateStyle.h"
+
+class FRiderSourceCodeAccessStyleSet : public FSlateStyleSet
+{
+public:
+
+	static const FRiderSourceCodeAccessStyleSet& Get();
+
+	static void Initialize();
+
+	static void Shutdown();
+
+	static FName RiderIconName;
+	static FName RiderRefreshIconName;
+	static FName StyleName;
+
+private:
+	FRiderSourceCodeAccessStyleSet();
+
+	virtual const FName& GetStyleSetName() const override;
+
+	static TUniquePtr<FRiderSourceCodeAccessStyleSet> Inst;
+	
+};
diff --git a/Engine/Plugins/Developer/VisualStudioCodeSourceCodeAccess/Source/VisualStudioCodeSourceCodeAccess/VisualStudioCodeSourceCodeAccess.Build.cs b/Engine/Plugins/Developer/VisualStudioCodeSourceCodeAccess/Source/VisualStudioCodeSourceCodeAccess/VisualStudioCodeSourceCodeAccess.Build.cs
--- a/Engine/Plugins/Developer/VisualStudioCodeSourceCodeAccess/Source/VisualStudioCodeSourceCodeAccess/VisualStudioCodeSourceCodeAccess.Build.cs
+++ b/Engine/Plugins/Developer/VisualStudioCodeSourceCodeAccess/Source/VisualStudioCodeSourceCodeAccess/VisualStudioCodeSourceCodeAccess.Build.cs
@@ -13,7 +13,7 @@ namespace UnrealBuildTool.Rules
 				{
 					"Core",
 					"SourceCodeAccess",
-					"DesktopPlatform",
+					"DesktopPlatform"
 				}
 			);
 
diff --git a/Engine/Source/Developer/SourceCodeAccess/Private/ISourceCodeAccessor.cpp b/Engine/Source/Developer/SourceCodeAccess/Private/ISourceCodeAccessor.cpp
new file mode 100644
--- /dev/null
+++ b/Engine/Source/Developer/SourceCodeAccess/Private/ISourceCodeAccessor.cpp
@@ -0,0 +1,8 @@
+﻿#include "ISourceCodeAccessor.h"
+
+#include "Styling/AppStyle.h"
+
+FName ISourceCodeAccessor::GetStyleSet() const
+{
+	return FAppStyle::GetAppStyleSetName();
+}
diff --git a/Engine/Source/Developer/SourceCodeAccess/Public/ISourceCodeAccessor.h b/Engine/Source/Developer/SourceCodeAccess/Public/ISourceCodeAccessor.h
--- a/Engine/Source/Developer/SourceCodeAccess/Public/ISourceCodeAccessor.h
+++ b/Engine/Source/Developer/SourceCodeAccess/Public/ISourceCodeAccessor.h
@@ -116,4 +116,11 @@ public:
 	{
 		return FName("MainFrame.RefreshSourceCodeEditor"); // Default Refresh Source Code Editor Icon
 	}
+	
+	/**
+	 * Get the style set where custom icons for the editors are defined.
+	 * By default uses FAppStyle::GetAppStyleSetName() == "CoreStyle"
+	 * @return the default style set for icons
+	 */
+	SOURCECODEACCESS_API virtual FName GetStyleSet() const;
 };
diff --git a/Engine/Source/Developer/SourceCodeAccess/SourceCodeAccess.Build.cs b/Engine/Source/Developer/SourceCodeAccess/SourceCodeAccess.Build.cs
--- a/Engine/Source/Developer/SourceCodeAccess/SourceCodeAccess.Build.cs
+++ b/Engine/Source/Developer/SourceCodeAccess/SourceCodeAccess.Build.cs
@@ -17,6 +17,7 @@ public class SourceCodeAccess : ModuleRules
 			new string[] {
 				"CoreUObject",
 				"Slate",
+				"SlateCore",
 			}
 		);
 	}
diff --git a/Engine/Source/Developer/TraceInsights/Private/Insights/MemoryProfiler/Widgets/SMemAllocTableTreeView.cpp b/Engine/Source/Developer/TraceInsights/Private/Insights/MemoryProfiler/Widgets/SMemAllocTableTreeView.cpp
--- a/Engine/Source/Developer/TraceInsights/Private/Insights/MemoryProfiler/Widgets/SMemAllocTableTreeView.cpp
+++ b/Engine/Source/Developer/TraceInsights/Private/Insights/MemoryProfiler/Widgets/SMemAllocTableTreeView.cpp
@@ -1632,7 +1632,7 @@ void SMemAllocTableTreeView::ExtendMenu(FMenuBuilder& MenuBuilder)
 				ItemToolTip,
 				FNewMenuDelegate::CreateSP(this, &SMemAllocTableTreeView::BuildOpenSourceSubMenu, true),
 				false,
-				FSlateIcon(FAppStyle::GetAppStyleSetName(), SourceCodeAccessor.GetOpenIconName())
+				FSlateIcon(SourceCodeAccessor.GetStyleSet(), SourceCodeAccessor.GetOpenIconName())
 			);
 
 			ItemLabel = FText::Format(LOCTEXT("ContextMenu_Open_SubMenu_Free", "Open in {0} | Free Callstack"), SourceCodeAccessor.GetNameText());
@@ -1645,7 +1645,7 @@ void SMemAllocTableTreeView::ExtendMenu(FMenuBuilder& MenuBuilder)
 				ItemToolTip,
 				FNewMenuDelegate::CreateSP(this, &SMemAllocTableTreeView::BuildOpenSourceSubMenu, false),
 				false,
-				FSlateIcon(FAppStyle::GetAppStyleSetName(), SourceCodeAccessor.GetOpenIconName())
+				FSlateIcon(SourceCodeAccessor.GetStyleSet(), SourceCodeAccessor.GetOpenIconName())
 			);
 		}
 		MenuBuilder.EndSection();
@@ -1667,7 +1667,7 @@ void SMemAllocTableTreeView::ExtendMenu(FMenuBuilder& MenuBuilder)
 			(
 				ItemLabel,
 				ItemToolTip,
-				FSlateIcon(FAppStyle::GetAppStyleSetName(), SourceCodeAccessor.GetOpenIconName()),
+				FSlateIcon(SourceCodeAccessor.GetStyleSet(), SourceCodeAccessor.GetOpenIconName()),
 				Action_OpenIDE,
 				NAME_None,
 				EUserInterfaceActionType::Button
diff --git a/Engine/Source/Developer/TraceInsights/Private/Insights/TaskGraphProfiler/Widgets/STaskTableTreeView.cpp b/Engine/Source/Developer/TraceInsights/Private/Insights/TaskGraphProfiler/Widgets/STaskTableTreeView.cpp
--- a/Engine/Source/Developer/TraceInsights/Private/Insights/TaskGraphProfiler/Widgets/STaskTableTreeView.cpp
+++ b/Engine/Source/Developer/TraceInsights/Private/Insights/TaskGraphProfiler/Widgets/STaskTableTreeView.cpp
@@ -119,7 +119,7 @@ void STaskTableTreeView::ExtendMenu(FMenuBuilder& MenuBuilder)
 			NAME_None,
 			TAttribute<FText>(),
 			TAttribute<FText>(),
-			FSlateIcon(FAppStyle::GetAppStyleSetName(), SourceCodeAccessor.GetOpenIconName())
+			FSlateIcon(SourceCodeAccessor.GetStyleSet(), SourceCodeAccessor.GetOpenIconName())
 		);
 	}
 
diff --git a/Engine/Source/Developer/TraceInsights/Private/Insights/Widgets/SLogView.cpp b/Engine/Source/Developer/TraceInsights/Private/Insights/Widgets/SLogView.cpp
--- a/Engine/Source/Developer/TraceInsights/Private/Insights/Widgets/SLogView.cpp
+++ b/Engine/Source/Developer/TraceInsights/Private/Insights/Widgets/SLogView.cpp
@@ -1245,7 +1245,7 @@ TSharedPtr<SWidget> SLogView::ListView_GetContextMenu()
 				NAME_None,
 				ItemLabel,
 				ItemToolTip,
-				FSlateIcon(FAppStyle::GetAppStyleSetName(), SourceCodeAccessor.GetOpenIconName())
+				FSlateIcon(SourceCodeAccessor.GetStyleSet(), SourceCodeAccessor.GetOpenIconName())
 			);
 		}
 	}
diff --git a/Engine/Source/Developer/TraceInsights/Private/Insights/Widgets/STimersView.cpp b/Engine/Source/Developer/TraceInsights/Private/Insights/Widgets/STimersView.cpp
--- a/Engine/Source/Developer/TraceInsights/Private/Insights/Widgets/STimersView.cpp
+++ b/Engine/Source/Developer/TraceInsights/Private/Insights/Widgets/STimersView.cpp
@@ -629,7 +629,7 @@ TSharedPtr<SWidget> STimersView::TreeView_GetMenuContent()
 				NAME_None,
 				ItemLabel,
 				ItemToolTip,
-				FSlateIcon(FAppStyle::GetAppStyleSetName(), SourceCodeAccessor.GetOpenIconName())
+				FSlateIcon(SourceCodeAccessor.GetStyleSet(), SourceCodeAccessor.GetOpenIconName())
 			);
 		}
 	}
diff --git a/Engine/Source/Editor/UnrealEd/Private/SourceCodeNavigation.cpp b/Engine/Source/Editor/UnrealEd/Private/SourceCodeNavigation.cpp
--- a/Engine/Source/Editor/UnrealEd/Private/SourceCodeNavigation.cpp
+++ b/Engine/Source/Editor/UnrealEd/Private/SourceCodeNavigation.cpp
@@ -1714,14 +1714,14 @@ FText FSourceCodeNavigation::GetSuggestedSourceCodeIDE(bool bShortIDEName)
 FSlateIcon FSourceCodeNavigation::GetOpenSourceCodeIDEIcon()
 {
 	ISourceCodeAccessModule& SourceCodeAccessModule = FModuleManager::LoadModuleChecked<ISourceCodeAccessModule>("SourceCodeAccess");
-	return FSlateIcon(FAppStyle::GetAppStyleSetName(), SourceCodeAccessModule.GetAccessor().GetOpenIconName());
+	return FSlateIcon(SourceCodeAccessModule.GetAccessor().GetStyleSet(), SourceCodeAccessModule.GetAccessor().GetOpenIconName());
 }
 
 
 FSlateIcon FSourceCodeNavigation::GetRefreshSourceCodeIDEIcon()
 {
 	ISourceCodeAccessModule& SourceCodeAccessModule = FModuleManager::LoadModuleChecked<ISourceCodeAccessModule>("SourceCodeAccess");
-	return FSlateIcon(FAppStyle::GetAppStyleSetName(), SourceCodeAccessModule.GetAccessor().GetRefreshIconName());
+	return FSlateIcon(SourceCodeAccessModule.GetAccessor().GetStyleSet(), SourceCodeAccessModule.GetAccessor().GetRefreshIconName());
 }
 
 FString FSourceCodeNavigation::GetSuggestedSourceCodeIDEDownloadURL()
diff --git a/Engine/Source/Runtime/Core/Core.Build.cs b/Engine/Source/Runtime/Core/Core.Build.cs
--- a/Engine/Source/Runtime/Core/Core.Build.cs
+++ b/Engine/Source/Runtime/Core/Core.Build.cs
@@ -40,8 +40,6 @@ public class Core : ModuleRules
 
 		if (Target.bBuildEditor == true)
 		{
-			DynamicallyLoadedModuleNames.Add("SourceCodeAccess");
-
 			PrivateIncludePathModuleNames.Add("DirectoryWatcher");
 			DynamicallyLoadedModuleNames.Add("DirectoryWatcher");
 		}
