Fix missing copyright boilerplate in source control files

diff --git a/Engine/Plugins/Developer/RiderSourceCodeAccess/Source/RiderSourceCodeAccess/Private/Style/RiderSourceCodeAccessStyleSet.cpp b/Engine/Plugins/Developer/RiderSourceCodeAccess/Source/RiderSourceCodeAccess/Private/Style/RiderSourceCodeAccessStyleSet.cpp
--- a/Engine/Plugins/Developer/RiderSourceCodeAccess/Source/RiderSourceCodeAccess/Private/Style/RiderSourceCodeAccessStyleSet.cpp
+++ b/Engine/Plugins/Developer/RiderSourceCodeAccess/Source/RiderSourceCodeAccess/Private/Style/RiderSourceCodeAccessStyleSet.cpp
@@ -1,4 +1,6 @@
-﻿#include "RiderSourceCodeAccessStyleSet.h"
+﻿// Copyright Epic Games, Inc. All Rights Reserved.
+
+#include "RiderSourceCodeAccessStyleSet.h"
 
 #include "Interfaces/IPluginManager.h"
 #include "Styling/SlateStyleRegistry.h"
diff --git a/Engine/Plugins/Developer/RiderSourceCodeAccess/Source/RiderSourceCodeAccess/Private/Style/RiderSourceCodeAccessStyleSet.h b/Engine/Plugins/Developer/RiderSourceCodeAccess/Source/RiderSourceCodeAccess/Private/Style/RiderSourceCodeAccessStyleSet.h
--- a/Engine/Plugins/Developer/RiderSourceCodeAccess/Source/RiderSourceCodeAccess/Private/Style/RiderSourceCodeAccessStyleSet.h
+++ b/Engine/Plugins/Developer/RiderSourceCodeAccess/Source/RiderSourceCodeAccess/Private/Style/RiderSourceCodeAccessStyleSet.h
@@ -1,4 +1,6 @@
-﻿#pragma once
+﻿// Copyright Epic Games, Inc. All Rights Reserved.
+
+#pragma once
 
 #include "Styling/SlateStyle.h"
 
