Fixed perf issue with mdfind taking multiple minutes due to it matching debug files emitted from Shader compilation. Query now takes less than a second.

diff --git a/Engine/Plugins/Developer/RiderSourceCodeAccess/Source/RiderSourceCodeAccess/Private/RiderPathLocator/Mac/RiderPathLocatorMac.cpp b/Engine/Plugins/Developer/RiderSourceCodeAccess/Source/RiderSourceCodeAccess/Private/RiderPathLocator/Mac/RiderPathLocatorMac.cpp
--- a/Engine/Plugins/Developer/RiderSourceCodeAccess/Source/RiderSourceCodeAccess/Private/RiderPathLocator/Mac/RiderPathLocatorMac.cpp
+++ b/Engine/Plugins/Developer/RiderSourceCodeAccess/Source/RiderSourceCodeAccess/Private/RiderPathLocator/Mac/RiderPathLocatorMac.cpp
@@ -80,21 +80,15 @@ static TArray<FInstallInfo> GetInstalledRidersWithMdfind()
 		return {};
 	}
 
-	FPlatformProcess::ExecProcess(TEXT("/usr/bin/mdfind"), TEXT("\"kMDItemKind == Application\""), &ReturnCode, &OutResults, &OutErrors);
+	FPlatformProcess::ExecProcess(TEXT("/usr/bin/mdfind"), TEXT("\"kMDItemContentType == com.apple.application-bundle && (kMDItemFSName == *Rider*)\""), &ReturnCode, &OutResults, &OutErrors);
 	if (ReturnCode != 0)
 	{
 		return {};
 	}
 
 	TArray<FString> RiderPaths;
-	FString TmpString;
-	while(OutResults.Split(TEXT("\n"), &TmpString, &OutResults))
-	{
-		if(TmpString.Contains(TEXT("Rider")))
-		{
-			RiderPaths.Add(TmpString);
-		}
-	}
+	OutResults.ParseIntoArray(RiderPaths, TEXT("\n"));
+	
 	TArray<FInstallInfo> Result;
 	for(const FString& RiderPath: RiderPaths)
 	{
